Linux: switch from LinuxNativeDialogs to SlateFileDialogs.

- Still experimental at this point.

diff --git a/Engine/Source/Developer/DesktopPlatform/DesktopPlatform.Build.cs b/Engine/Source/Developer/DesktopPlatform/DesktopPlatform.Build.cs
--- a/Engine/Source/Developer/DesktopPlatform/DesktopPlatform.Build.cs
+++ b/Engine/Source/Developer/DesktopPlatform/DesktopPlatform.Build.cs
@@ -17,8 +17,8 @@ public class DesktopPlatform : ModuleRules
 
 		if (Target.Platform == UnrealTargetPlatform.Linux)
 		{
-            AddThirdPartyPrivateStaticDependencies(Target, "LinuxNativeDialogs");
-            AddThirdPartyPrivateStaticDependencies(Target, "SDL2");
+			AddThirdPartyPrivateStaticDependencies(Target, "SDL2");
+			AddThirdPartyPrivateStaticDependencies(Target, "SlateFileDialogs");
 		}
 	}
 }
diff --git a/Engine/Source/Developer/DesktopPlatform/Private/Linux/DesktopPlatformLinux.cpp b/Engine/Source/Developer/DesktopPlatform/Private/Linux/DesktopPlatformLinux.cpp
--- a/Engine/Source/Developer/DesktopPlatform/Private/Linux/DesktopPlatformLinux.cpp
+++ b/Engine/Source/Developer/DesktopPlatform/Private/Linux/DesktopPlatformLinux.cpp
@@ -6,7 +6,9 @@
 //#include "LinuxNativeFeedbackContext.h"
 // custom dialogs
 #if WITH_LINUX_NATIVE_DIALOGS
-#include "UNativeDialogs.h"
+	#include "UNativeDialogs.h"
+#else
+	#include "SlateFileDialogs.h"
 #endif // WITH_LINUX_NATIVE_DIALOGS
 #include "SDL.h"
 
@@ -28,8 +30,8 @@ FDesktopPlatformLinux::FDesktopPlatformLinux()
 		UE_LOG(LogDesktopPlatform, Warning, TEXT("DesktopPlatformLinux could not initialize LinuxNativeDialogs - it will not work properly."));
 	}
 #else
-	UE_LOG(LogDesktopPlatform, Warning, TEXT("DesktopPlatformLinux is not using LinuxNativeDialogs - it will not work properly."));
-#endif
+	UE_LOG(LogDesktopPlatform, Log, TEXT("DesktopPlatformLinux is not using LinuxNativeDialogs."));
+#endif // WITH_LINUX_NATIVE_DIALOGS
 }
 
 FDesktopPlatformLinux::~FDesktopPlatformLinux()
@@ -41,22 +43,68 @@ FDesktopPlatformLinux::~FDesktopPlatformLinux()
 
 bool FDesktopPlatformLinux::OpenFileDialog(const void* ParentWindowHandle, const FString& DialogTitle, const FString& DefaultPath, const FString& DefaultFile, const FString& FileTypes, uint32 Flags, TArray<FString>& OutFilenames, int32& OutFilterIndex)
 {
+#if WITH_LINUX_NATIVE_DIALOGS
 	return FileDialogShared(false, ParentWindowHandle, DialogTitle, DefaultPath, DefaultFile, FileTypes, Flags, OutFilenames, OutFilterIndex);
+#else
+	if (!FModuleManager::Get().IsModuleLoaded("SlateFileDialogs"))
+	{
+		FModuleManager::Get().LoadModule("SlateFileDialogs");
+	}
+
+	ISlateFileDialogsModule *FileDialog = FModuleManager::GetModulePtr<ISlateFileDialogsModule>("SlateFileDialogs");
+
+	if (FileDialog)
+	{
+		return FileDialog->OpenFileDialog(ParentWindowHandle, DialogTitle, DefaultPath, DefaultFile, FileTypes, Flags, OutFilenames, OutFilterIndex);
+	}
+
+	return false;
+#endif // WITH_LINUX_NATIVE_DIALOGS
 }
 
 bool FDesktopPlatformLinux::OpenFileDialog(const void* ParentWindowHandle, const FString& DialogTitle, const FString& DefaultPath, const FString& DefaultFile, const FString& FileTypes, uint32 Flags, TArray<FString>& OutFilenames)
 {
+#if WITH_LINUX_NATIVE_DIALOGS
 	int32 DummyFilterIndex;
 	return FileDialogShared(false, ParentWindowHandle, DialogTitle, DefaultPath, DefaultFile, FileTypes, Flags, OutFilenames, DummyFilterIndex);
+#else
+	if (!FModuleManager::Get().IsModuleLoaded("SlateFileDialogs"))
+	{
+		FModuleManager::Get().LoadModule("SlateFileDialogs");
+	}
+
+	ISlateFileDialogsModule *FileDialog = FModuleManager::GetModulePtr<ISlateFileDialogsModule>("SlateFileDialogs");
+
+	if (FileDialog)
+	{
+		return FileDialog->OpenFileDialog(ParentWindowHandle, DialogTitle, DefaultPath, DefaultFile, FileTypes, Flags, OutFilenames);
+	}
+
+	return false;
+#endif // WITH_LINUX_NATIVE_DIALOGS
 }
 
 bool FDesktopPlatformLinux::SaveFileDialog(const void* ParentWindowHandle, const FString& DialogTitle, const FString& DefaultPath, const FString& DefaultFile, const FString& FileTypes, uint32 Flags, TArray<FString>& OutFilenames)
 {
+#if WITH_LINUX_NATIVE_DIALOGS
 	int32 DummyFilterIndex = 0;
 	return FileDialogShared(true, ParentWindowHandle, DialogTitle, DefaultPath, DefaultFile, FileTypes, Flags, OutFilenames, DummyFilterIndex);
-}
+#else
+	if (!FModuleManager::Get().IsModuleLoaded("SlateFileDialogs"))
+	{
+		FModuleManager::Get().LoadModule("SlateFileDialogs");
+	}
 
+	ISlateFileDialogsModule *FileDialog = FModuleManager::GetModulePtr<ISlateFileDialogsModule>("SlateFileDialogs");
 
+	if (FileDialog)
+	{
+		return FileDialog->SaveFileDialog(ParentWindowHandle, DialogTitle, DefaultPath, DefaultFile, FileTypes, Flags, OutFilenames);
+	}
+
+	return false;
+#endif // WITH_LINUX_NATIVE_DIALOGS
+}
 
 bool FDesktopPlatformLinux::OpenDirectoryDialog(const void* ParentWindowHandle, const FString& DialogTitle, const FString& DefaultPath, FString& OutFolderName)
 {
@@ -103,6 +151,18 @@ bool FDesktopPlatformLinux::OpenDirectoryDialog(const void* ParentWindowHandle,
 
 	return bSuccess;
 #else
+	if (!FModuleManager::Get().IsModuleLoaded("SlateFileDialogs"))
+	{
+		FModuleManager::Get().LoadModule("SlateFileDialogs");
+	}
+
+	ISlateFileDialogsModule *FileDialog = FModuleManager::GetModulePtr<ISlateFileDialogsModule>("SlateFileDialogs");
+
+	if (FileDialog)
+	{
+		return FileDialog->OpenDirectoryDialog(ParentWindowHandle, DialogTitle, DefaultPath, OutFolderName);
+	}
+
 	return false;
 #endif // WITH_LINUX_NATIVE_DIALOGS
 }
diff --git a/Engine/Source/Developer/SlateFileDialogs/Private/SlateFileDialogsModule.cpp b/Engine/Source/Developer/SlateFileDialogs/Private/SlateFileDialogsModule.cpp
new file mode 100644
--- /dev/null
+++ b/Engine/Source/Developer/SlateFileDialogs/Private/SlateFileDialogsModule.cpp
@@ -0,0 +1,59 @@
+// Copyright 1998-2015 Epic Games, Inc. All Rights Reserved.
+
+#include "SlateFileDialogsPrivatePCH.h"
+#include "SlateFileDialogs.h"
+#include "SlateFileDlgWindow.h"
+
+void FSlateFileDialogsModule::StartupModule()
+{
+	SlateFileDialog = new FSlateFileDialogsModule();
+}
+
+
+void FSlateFileDialogsModule::ShutdownModule()
+{
+	if (SlateFileDialog != NULL)
+	{
+		delete SlateFileDialog;
+		SlateFileDialog = NULL;
+	}
+}
+
+
+bool FSlateFileDialogsModule::OpenFileDialog(const void* ParentWindowHandle, const FString& DialogTitle, const FString& DefaultPath,
+		const FString& DefaultFile, const FString& FileTypes, uint32 Flags, TArray<FString>& OutFilenames, int32& OutFilterIndex)
+{
+	FSlateFileDlgWindow dialog;
+	return dialog.OpenFileDialog(ParentWindowHandle, DialogTitle, DefaultPath, DefaultFile, FileTypes, Flags, OutFilenames, OutFilterIndex);
+}
+
+
+bool FSlateFileDialogsModule::OpenFileDialog(const void* ParentWindowHandle, const FString& DialogTitle, const FString& DefaultPath,
+		const FString& DefaultFile, const FString& FileTypes, uint32 Flags, TArray<FString>& OutFilenames)
+{
+	FSlateFileDlgWindow dialog;
+	return dialog.OpenFileDialog(ParentWindowHandle, DialogTitle, DefaultPath, DefaultFile, FileTypes, Flags, OutFilenames);
+}
+
+bool FSlateFileDialogsModule::OpenDirectoryDialog(const void* ParentWindowHandle, const FString& DialogTitle, const FString& DefaultPath,
+		FString& OutFoldername)
+{
+	FSlateFileDlgWindow dialog;
+	return dialog.OpenDirectoryDialog(ParentWindowHandle, DialogTitle, DefaultPath, OutFoldername);
+}
+
+
+bool FSlateFileDialogsModule::SaveFileDialog(const void* ParentWindowHandle, const FString& DialogTitle, const FString& DefaultPath,
+	const FString& DefaultFile, const FString& FileTypes, uint32 Flags, TArray<FString>& OutFilenames)
+{
+	FSlateFileDlgWindow dialog;
+	return dialog.SaveFileDialog(ParentWindowHandle, DialogTitle, DefaultPath, DefaultFile, FileTypes, Flags, OutFilenames);
+}
+
+
+ISlateFileDialogsModule* FSlateFileDialogsModule::Get()
+{
+	return SlateFileDialog;
+}
+
+IMPLEMENT_MODULE(FSlateFileDialogsModule, SlateFileDialogs);
diff --git a/Engine/Source/Developer/SlateFileDialogs/Private/SlateFileDialogsPrivatePCH.h b/Engine/Source/Developer/SlateFileDialogs/Private/SlateFileDialogsPrivatePCH.h
new file mode 100644
--- /dev/null
+++ b/Engine/Source/Developer/SlateFileDialogs/Private/SlateFileDialogsPrivatePCH.h
@@ -0,0 +1,9 @@
+// Copyright 1998-2015 Epic Games, Inc. All Rights Reserved.
+
+#pragma once
+
+#define HIDE_HIDDEN_FILES 1
+#define MAX_FILTER_LENGTH 8192
+#define MAX_PATH_LENGTH 8192
+
+#define ENABLE_DIRECTORY_WATCHER 1
diff --git a/Engine/Source/Developer/SlateFileDialogs/Private/SlateFileDlgWindow.cpp b/Engine/Source/Developer/SlateFileDialogs/Private/SlateFileDlgWindow.cpp
new file mode 100644
--- /dev/null
+++ b/Engine/Source/Developer/SlateFileDialogs/Private/SlateFileDlgWindow.cpp
@@ -0,0 +1,1262 @@
+// Copyright 1998-2015 Epic Games, Inc. All Rights Reserved.
+
+#include "SlateFileDialogsPrivatePCH.h"
+#include "SlateFileDlgWindow.h"
+#include "ModuleManager.h"
+
+#define LOCTEXT_NAMESPACE "SlateFileDialogsNamespace"
+
+DEFINE_LOG_CATEGORY_STATIC(LogSlateFileDialogs, Log, All);
+
+//-----------------------------------------------------------------------------
+//-----------------------------------------------------------------------------
+
+class FSlateFileDialogVisitor : public IPlatformFile::FDirectoryVisitor
+{
+public:
+	FSlateFileDialogVisitor(TArray<TSharedPtr<FFileEntry>> &InFileList,
+						TArray<TSharedPtr<FFileEntry>> &InFolderList, TSharedPtr<FString> InFilterList)
+		: FileList(InFileList),
+		FolderList(InFolderList),
+		FilterList(InFilterList)
+	{}
+	
+
+	virtual bool Visit(const TCHAR* FilenameOrDirectory, bool bIsDirectory) override
+	{
+		int32 i;
+
+		// break filename from path
+		for (i = FCString::Strlen(FilenameOrDirectory) - 1; i >= 0; i--)
+		{
+			if (FilenameOrDirectory[i] == TCHAR('/'))
+				break;
+		}
+
+#if HIDE_HIDDEN_FILES
+		if (FilenameOrDirectory[i + 1] == TCHAR('.'))
+		{
+			return true;
+		}
+#endif
+
+		FDateTime stamp = IFileManager::Get().GetTimeStamp(FilenameOrDirectory);
+		FString ModDate = "";		
+		FString FileSize = "";
+
+		if (bIsDirectory)
+		{
+			FolderList.Add(MakeShareable(new FFileEntry(FString(&FilenameOrDirectory[i + 1]), ModDate, FileSize, true)));
+		}
+		else
+		{
+			if (PassesFilterTest(&FilenameOrDirectory[i + 1]))
+			{
+				int64 size = IFileManager::Get().FileSize(FilenameOrDirectory);
+
+				if (size < 1048576)
+				{
+					size = (size + 1023) / 1024;
+					FileSize = FString::FromInt(size) + " KB";
+				}
+				else
+				{
+					size /= 1024;
+
+					if (size < 1048576)
+					{
+						size = (size + 1023) / 1024;
+						FileSize = FString::FromInt(size) + " MB";
+					}
+					else
+					{
+						size /= 1024;
+
+						size = (size + 1023) / 1024;
+						FileSize = FString::FromInt(size) + " GB";
+					}
+				}
+				
+				
+				ModDate = FString::Printf(TEXT("%02d/%02d/%04d "), stamp.GetMonth(), stamp.GetDay(), stamp.GetYear());
+				
+				if (stamp.GetHour() == 0)
+				{
+					ModDate = ModDate + FString::Printf(TEXT("12:%02d AM"), stamp.GetMinute());
+				}
+				else if (stamp.GetHour() < 12)
+				{
+					ModDate = ModDate + FString::Printf(TEXT("%2d:%02d AM"), stamp.GetHour12(), stamp.GetMinute());
+				}
+				else
+				{
+					ModDate = ModDate + FString::Printf(TEXT("%2d:%02d PM"), stamp.GetHour12(), stamp.GetMinute());
+				}
+
+				FileList.Add(MakeShareable(new FFileEntry(FString(&FilenameOrDirectory[i + 1]), ModDate, FileSize, false)));
+			}
+		}
+
+		return true;
+	}
+
+	bool PassesFilterTest(const TCHAR *Filename)
+	{		
+		if (!FilterList.IsValid())
+		{
+			return true; // no filters. everything passes.
+		}
+
+		int32 i;
+		const TCHAR *Extension = NULL;
+
+		// get extension
+		for (i = FCString::Strlen(Filename) - 1; i >= 0; i--)
+		{
+			if (Filename[i] == TCHAR('.'))
+			{
+				Extension = &Filename[i];
+				break;
+			}
+		}
+
+		if (Extension == NULL)
+		{
+			return false; // file has no extension. it fails filter test.
+		}
+
+		TCHAR Temp[MAX_FILTER_LENGTH];
+		FCString::Strcpy(Temp, FilterList->Len(), *(*FilterList.Get()));
+
+		// break path into tokens
+		TCHAR *ContextStr = nullptr;
+
+		TCHAR *FilterExt = FCString::Strtok(Temp, TEXT(";"), &ContextStr);
+		bool RC = false;
+
+		while (FilterExt)
+		{
+			// strip leading spaces and '*' from beginning.
+			FilterExt = FCString::Strchr(FilterExt, '.');
+
+			// strip any trailing spaces
+			for (i = FCString::Strlen(FilterExt) - 1; i >= 0 && FilterExt[i] == TCHAR(' '); i--)
+			{
+				FilterExt[i] = 0;
+			}
+
+			if (FCString::Strcmp(FilterExt, TEXT(".*")) == 0)
+			{
+				// *.* matches all
+				RC = true;
+				break;
+			}
+
+			if (FCString::Strcmp(FilterExt, Extension) == 0)
+			{
+				// positive hit.
+				RC = true;
+				break;
+			}
+
+			// next filter entry
+			FilterExt = FCString::Strtok(nullptr, TEXT(";"), &ContextStr);
+		}
+
+		// cleanup and return failed.
+		return RC;
+	}
+
+private:
+	TArray<TSharedPtr<FFileEntry>>& FileList;
+	TArray<TSharedPtr<FFileEntry>>& FolderList;
+	TSharedPtr<FString> FilterList;
+};
+
+
+class FSlateFileDialogDirVisitor : public IPlatformFile::FDirectoryVisitor
+{
+public:
+	FSlateFileDialogDirVisitor(TArray<FString> *InDirectoryNames)
+		: DirectoryNames(InDirectoryNames)
+	{}
+
+	void SetResultPath(TArray<FString> *InDirectoryNames) { DirectoryNames = InDirectoryNames; }
+
+	virtual bool Visit(const TCHAR* FilenameOrDirectory, bool bIsDirectory) override
+	{
+		int32 i;
+
+		// break filename from path
+		for (i = FCString::Strlen(FilenameOrDirectory) - 1; i >= 0; i--)
+		{
+			if (FilenameOrDirectory[i] == TCHAR('/'))
+				break;
+		}
+
+#if HIDE_HIDDEN_FILES
+		if (FilenameOrDirectory[i + 1] == TCHAR('.'))
+		{
+			return true;
+		}
+#endif
+
+		if (bIsDirectory)
+		{
+			DirectoryNames->Add(FString(&FilenameOrDirectory[i + 1]));
+		}
+
+		return true;
+	}
+	
+private:
+	TArray<FString> *DirectoryNames;
+};
+
+
+//-----------------------------------------------------------------------------
+//-----------------------------------------------------------------------------
+
+bool FSlateFileDlgWindow::OpenFileDialog(const void* ParentWindowHandle, const FString& DialogTitle, const FString& DefaultPath,
+		const FString& DefaultFile, const FString& FileTypes, uint32 Flags, TArray<FString>& OutFilenames, int32& OutFilterIndex)
+{
+	TSharedRef<SWindow> ModalWindow = SNew(SWindow)
+		.SupportsMinimize(false)
+		.SupportsMaximize(false)
+		.Title(LOCTEXT("SlateFileDialogsOpenFile","Open File"))
+		.CreateTitleBar(true)
+		.MinHeight(400.0f)
+		.MinWidth(600.0f)
+		.ActivateWhenFirstShown(true)
+		.ClientSize(FVector2D(800, 500));
+	
+	DialogWidget = SNew(SSlateFileOpenDlg)
+		.bMultiSelectEnabled(Flags == 1)
+		.ParentWindow(ModalWindow)
+		.CurrentPath(DefaultPath)
+		.Filters(FileTypes)
+		.WindowTitleText(DialogTitle);
+	
+	DialogWidget->SetOutNames(&OutFilenames);
+	DialogWidget->SetOutFilterIndex(&OutFilterIndex);
+	
+	ModalWindow->SetContent( DialogWidget.ToSharedRef() );
+		
+	FSlateApplication::Get().AddModalWindow(ModalWindow, NULL);
+	return (DialogWidget->GetResponse() == EResult::Accept && OutFilenames.Num() > 0);
+}
+
+
+bool FSlateFileDlgWindow::OpenFileDialog(const void* ParentWindowHandle, const FString& DialogTitle, const FString& DefaultPath,
+	const FString& DefaultFile, const FString& FileTypes, uint32 Flags, TArray<FString>& OutFilenames)
+{
+	int32 DummyIndex;
+	return OpenFileDialog(ParentWindowHandle, DialogTitle, DefaultPath, DefaultFile, FileTypes, Flags, OutFilenames, DummyIndex);
+}
+
+
+bool FSlateFileDlgWindow::OpenDirectoryDialog(const void* ParentWindowHandle, const FString& DialogTitle, const FString& DefaultPath,
+	FString& OutFoldername)
+{
+	int32 DummyIndex;
+	TArray<FString> TempOut;
+	FString Filters = "";
+
+	TSharedRef<SWindow> ModalWindow = SNew(SWindow)
+		.SupportsMinimize(false)
+		.SupportsMaximize(false)
+		.Title(LOCTEXT("SlateFileDialogsOpenDirectory","Open Directory"))
+		.CreateTitleBar(true)
+		.MinHeight(400.0f)
+		.MinWidth(600.0f)
+		.ActivateWhenFirstShown(true)
+		.ClientSize(FVector2D(800, 500));
+
+	DialogWidget = SNew(SSlateFileOpenDlg)
+		.bMultiSelectEnabled(false)
+		.ParentWindow(ModalWindow)
+		.bDirectoriesOnly(true)
+		.CurrentPath(DefaultPath)			
+		.WindowTitleText(DialogTitle);
+
+	DialogWidget->SetOutNames(&TempOut);
+	DialogWidget->SetOutFilterIndex(&DummyIndex);
+
+	ModalWindow->SetContent( DialogWidget.ToSharedRef() );
+
+	FSlateApplication::Get().AddModalWindow(ModalWindow, NULL);
+	bool RC = (DialogWidget->GetResponse() == EResult::Accept && TempOut.Num() > 0);
+
+	if (TempOut.Num() > 0)
+	{
+		OutFoldername = TempOut[0] + TEXT("/");
+	}
+
+	return RC;
+}
+
+
+bool FSlateFileDlgWindow::SaveFileDialog(const void* ParentWindowHandle, const FString& DialogTitle, const FString& DefaultPath,
+	const FString& DefaultFile, const FString& FileTypes, uint32 Flags, TArray<FString>& OutFilenames)
+{
+	int32 DummyIndex;
+
+	TSharedRef<SWindow> ModalWindow = SNew(SWindow)
+		.SupportsMinimize(false)
+		.SupportsMaximize(false)
+		.Title(LOCTEXT("SlateFileDialogsSaveFile","Save File"))
+		.CreateTitleBar(true)
+		.MinHeight(400.0f)
+		.MinWidth(600.0f)
+		.ActivateWhenFirstShown(true)
+		.ClientSize(FVector2D(800, 500));
+
+	DialogWidget = SNew(SSlateFileOpenDlg)
+		.bMultiSelectEnabled(false)
+		.ParentWindow(ModalWindow)
+		.bSaveFile(true)
+		.AcceptText(LOCTEXT("SlateFileDialogsSave","Save"))
+		.CurrentPath(DefaultPath)
+		.Filters(FileTypes)			
+		.WindowTitleText(DialogTitle);
+
+	DialogWidget->SetOutNames(&OutFilenames);
+	DialogWidget->SetOutFilterIndex(&DummyIndex);
+	DialogWidget->SetDefaultFile(DefaultFile);
+
+	ModalWindow->SetContent( DialogWidget.ToSharedRef() );
+		
+	FSlateApplication::Get().AddModalWindow(ModalWindow, NULL);
+	return (DialogWidget->GetResponse() == EResult::Accept && OutFilenames.Num() > 0);
+}
+
+
+
+//-----------------------------------------------------------------------------
+// custom file dialog widget
+//-----------------------------------------------------------------------------
+BEGIN_SLATE_FUNCTION_BUILD_OPTIMIZATION
+void SSlateFileOpenDlg::Construct(const FArguments& InArgs)
+{
+	CurrentPath = InArgs._CurrentPath;
+	Filters = InArgs._Filters;
+	bMultiSelectEnabled = InArgs._bMultiSelectEnabled;
+	bDirectoriesOnly = InArgs._bDirectoriesOnly;
+	bSaveFile = InArgs._bSaveFile;
+	WindowTitleText = InArgs._WindowTitleText;		
+	OutNames = InArgs._OutNames;
+	OutFilterIndex = InArgs._OutFilterIndex;
+	UserResponse = FSlateFileDlgWindow::Cancel;
+	ParentWindow = InArgs._ParentWindow;
+	AcceptText = InArgs._AcceptText;
+	DirNodeIndex = -1;
+	FilterIndex = 0;
+
+	ESelectionMode::Type SelectMode = bMultiSelectEnabled.Get() ? ESelectionMode::Multi : ESelectionMode::Single;
+	struct EVisibility SaveFilenameVisibility = bSaveFile.Get() ? EVisibility::Visible : EVisibility::Collapsed;
+
+	this->ChildSlot
+		[
+			SNew(SBorder)
+			.HAlign(HAlign_Fill)
+			.VAlign(VAlign_Fill)
+			.Padding(FMargin(20.0f,20.0f))
+			.BorderImage(FCoreStyle::Get().GetBrush("SlateFileDialogs.GroupBorder"))
+			[
+				SNew(SVerticalBox)
+
+				+ SVerticalBox::Slot() // window title
+				.HAlign(HAlign_Left)
+				.VAlign(VAlign_Fill)
+				.AutoHeight()
+				.Padding(FMargin(0.0f, 0.0f, 0.0f, 20.0f))
+				[
+					SAssignNew(WindowTitle, STextBlock)
+					.Text(WindowTitleText)
+					.Font(FCoreStyle::Get().GetFontStyle("SlateFileDialogs.DialogLarge"))
+					.Justification(ETextJustify::Center)
+				]
+
+				+ SVerticalBox::Slot() // Path breadcrumbs
+				.HAlign(HAlign_Fill)
+				.VAlign(VAlign_Fill)
+				.AutoHeight()
+				.Padding(FMargin(0.0f, 0.0f, 0.0f, 20.0f))
+				[
+					SAssignNew(PathBreadcrumbTrail, SBreadcrumbTrail<FString>)
+					.ButtonContentPadding(FMargin(2.0f, 2.0f))
+					.ButtonStyle(FCoreStyle::Get(), "SlateFileDialogs.FlatButton")
+					.DelimiterImage(FCoreStyle::Get().GetBrush("SlateFileDialogs.PathDelimiter"))
+					.TextStyle(FCoreStyle::Get(), "SlateFileDialogs.PathText")
+					.ShowLeadingDelimiter(false)
+					.InvertTextColorOnHover(false)
+					.OnCrumbClicked(this, &SSlateFileOpenDlg::OnPathClicked)
+					.GetCrumbMenuContent(this, &SSlateFileOpenDlg::OnGetCrumbDelimiterContent)
+					.AddMetaData<FTagMetaData>(FTagMetaData(TEXT("ContentBrowserPath")))
+				]
+
+				+ SVerticalBox::Slot()
+				.HAlign(HAlign_Fill)
+				.VAlign(VAlign_Fill)
+				.FillHeight(1.0f)
+				[
+					SNew(SHorizontalBox)
+
+					+ SHorizontalBox::Slot()
+					.Padding(FMargin(0.0f))
+					.AutoWidth()
+					.HAlign(HAlign_Fill)
+					.VAlign(VAlign_Fill)
+					[
+						SNew(SVerticalBox)
+
+						+ SVerticalBox::Slot()
+						.Padding(FMargin(10.0f))
+						.AutoHeight()
+						.HAlign(HAlign_Fill)
+						.VAlign(VAlign_Fill)
+						[
+							SNew(SButton)
+							.HAlign(HAlign_Center)
+							.VAlign(VAlign_Center)
+							.OnClicked(this, &SSlateFileOpenDlg::OnQuickLinkClick, FSlateFileDlgWindow::Project)
+							.ContentPadding(FMargin(2.0f))
+							[
+								SNew(SHorizontalBox)
+			
+								+ SHorizontalBox::Slot()
+								.HAlign(HAlign_Center)
+								.VAlign(VAlign_Center)
+								.AutoWidth()
+								[
+									SNew(SImage)
+									.Image(FCoreStyle::Get().GetBrush("SlateFileDialogs.Folder24"))
+								]
+			
+								+ SHorizontalBox::Slot()
+								.HAlign(HAlign_Center)
+								.VAlign(VAlign_Center)
+								.AutoWidth()				
+								[
+									SNew(STextBlock)
+									.Text(FString("Projects"))
+									.Font(FCoreStyle::Get().GetFontStyle("SlateFileDialogs.Dialog"))
+									.Justification(ETextJustify::Left)
+								]
+							]
+						]
+
+						+ SVerticalBox::Slot()
+						.Padding(FMargin(10.0f))
+						.AutoHeight()
+						.HAlign(HAlign_Fill)
+						.VAlign(VAlign_Fill)
+						[
+							SNew(SButton)
+							.HAlign(HAlign_Center)
+							.VAlign(VAlign_Center)
+							.OnClicked(this, &SSlateFileOpenDlg::OnQuickLinkClick, FSlateFileDlgWindow::Engine)
+							.ContentPadding(FMargin(2.0f))
+							[
+								SNew(SHorizontalBox)
+
+								+ SHorizontalBox::Slot()
+								.HAlign(HAlign_Center)
+								.VAlign(VAlign_Center)
+								.AutoWidth()
+								[
+									SNew(SImage)
+									.Image(FCoreStyle::Get().GetBrush("SlateFileDialogs.Folder24"))
+								]
+
+								+ SHorizontalBox::Slot()
+								.HAlign(HAlign_Center)
+								.VAlign(VAlign_Center)
+								.AutoWidth()				
+								[
+									SNew(STextBlock)
+									.Text(FString("Engine"))
+									.Font(FCoreStyle::Get().GetFontStyle("SlateFileDialogs.Dialog"))
+									.Justification(ETextJustify::Left)
+								]
+							]
+						]
+					]
+
+					+ SHorizontalBox::Slot() // spacer
+					.Padding(FMargin(0.0f))
+					.AutoWidth()
+					.HAlign(HAlign_Fill)
+					.VAlign(VAlign_Fill)
+					[
+						SNew(SSpacer)
+						.Size(FVector2D(20.0f, 1.0f))
+					]
+
+					+ SHorizontalBox::Slot() // file list area
+					.Padding(FMargin(0.0f, 0.0f, 20.0f, 0.0f))
+					.HAlign(HAlign_Fill)
+					.VAlign(VAlign_Fill)
+					.FillWidth(1.0f)
+					[
+						SNew(SBorder)
+						.HAlign(HAlign_Fill)
+						.VAlign(VAlign_Fill)
+						.Padding(FMargin(10.0f))
+						.BorderBackgroundColor(FLinearColor(0.10f, 0.10f, 0.10f, 1.0f))
+						.BorderImage(FCoreStyle::Get().GetBrush("SlateFileDialogs.WhiteBackground"))
+						[
+							SAssignNew(ListView, SListView<TSharedPtr<FFileEntry>>) // file list scroll
+							.ListItemsSource(&LineItemArray)
+							.SelectionMode(SelectMode)
+							.OnGenerateRow(this, &SSlateFileOpenDlg::OnGenerateWidgetForList)
+							.OnMouseButtonDoubleClick(this, &SSlateFileOpenDlg::OnItemDoubleClicked)
+							.OnSelectionChanged(this, &SSlateFileOpenDlg::OnItemSelected)
+							.HeaderRow
+							(
+								SNew(SHeaderRow)
+								.Visibility(EVisibility::Visible)
+								
+								+ SHeaderRow::Column("Pathname")
+									.DefaultLabel(LOCTEXT("SlateFileDialogsNameHeader", "Name"))
+									.FillWidth(1.0f)
+								
+								+ SHeaderRow::Column("ModDate")
+									.DefaultLabel(LOCTEXT("SlateFileDialogsModDateHeader", "Date Modified"))
+									.FixedWidth(170.0f)
+
+								+ SHeaderRow::Column("FileSize")
+									.DefaultLabel(LOCTEXT("SlateFileDialogsFileSizeHeader", "File Size"))
+									.FixedWidth(70.0f)
+							)
+						]
+					]
+				]
+
+				+ SVerticalBox::Slot() // save filename entry
+				.HAlign(HAlign_Right)
+				.VAlign(VAlign_Bottom)
+				.Padding(FMargin(0.0f, 10.0f, 50.0f, 0.0f))
+				.AutoHeight()
+				[
+					SAssignNew(SaveFilenameSizeBox, SBox)
+					.Padding(FMargin(0.0f))
+					.HAlign(HAlign_Fill)
+					.VAlign(VAlign_Fill)
+					.MinDesiredHeight(20.0f)
+					.Visibility(SaveFilenameVisibility)
+					[
+						SNew(SHorizontalBox)
+
+						+ SHorizontalBox::Slot()
+						.Padding(FMargin(0.0f))
+						.AutoWidth()
+						.HAlign(HAlign_Fill)
+						.VAlign(VAlign_Fill)
+						[
+							SNew(STextBlock)
+							.Text(FString("Save Filename:"))
+							.Font(FCoreStyle::Get().GetFontStyle("SlateFileDialogs.Dialog"))
+							.Justification(ETextJustify::Left)
+						]
+
+						+ SHorizontalBox::Slot()
+						.Padding(FMargin(10.0f, 0.0f, 0.0f, 0.0f))
+						.AutoWidth()
+						.HAlign(HAlign_Fill)
+						.VAlign(VAlign_Fill)
+						[
+							SNew(SBox)
+							.Padding(FMargin(0.0f))
+							.HAlign(HAlign_Fill)
+							.VAlign(VAlign_Fill)
+							.WidthOverride(300.0f)
+							[
+								SNew(SBorder)
+								.HAlign(HAlign_Fill)
+								.VAlign(VAlign_Fill)
+								.Padding(FMargin(5.0f,0.0f))
+								.BorderBackgroundColor(FLinearColor(0.1f, 0.1f, 0.1f, 1.0f))
+								.BorderImage(FCoreStyle::Get().GetBrush("SlateFileDialogs.WhiteBackground"))
+								[
+									SAssignNew(SaveFilenameEditBox, SInlineEditableTextBlock)
+									.Font(FCoreStyle::Get().GetFontStyle("SlateFileDialogs.Dialog"))
+									.IsReadOnly(false)
+									.Text(FText::FromString(""))
+									.OnTextCommitted(this, &SSlateFileOpenDlg::OnFileNameCommitted)
+								]
+							]
+						]
+					]
+				]
+
+				+ SVerticalBox::Slot()  // cancel:accept buttons
+				.HAlign(HAlign_Fill)
+				.VAlign(VAlign_Fill)
+				.AutoHeight()
+				.Padding(FMargin(0.0f, 10.0f, 0.0f, 0.0f))
+				[
+					SNew(SHorizontalBox)
+
+					+ SHorizontalBox::Slot()
+					.Padding(FMargin(0.0f))
+					.AutoWidth()
+					.HAlign(HAlign_Left)
+					.VAlign(VAlign_Top)
+					[
+						SAssignNew(FilterHBox, SHorizontalBox)
+
+						+ SHorizontalBox::Slot()
+						.HAlign(HAlign_Left)
+						.VAlign(VAlign_Bottom)
+						.AutoWidth()
+						.Padding(FMargin(0.0f))
+						[
+							SNew(STextBlock)
+							.Text(FString("Filter:"))
+							.Font(FCoreStyle::Get().GetFontStyle("SlateFileDialogs.Dialog"))
+							.Justification(ETextJustify::Left)
+						]
+
+						+ SHorizontalBox::Slot()
+						.HAlign(HAlign_Fill)
+						.VAlign(VAlign_Fill)
+						.AutoWidth()
+						.Padding(FMargin(10.0f, 0.0f, 0.0f, 0.0f))
+						[
+							SNew(SBox)
+							.MinDesiredWidth(200.0f)
+							.MaxDesiredWidth(200.0f)
+							.Padding(FMargin(0.0f))
+							[
+								SAssignNew(FilterCombo, STextComboBox)
+								.ContentPadding(FMargin(4.0f, 2.0f))
+								//.MaxListHeight(100.0f)
+								.OptionsSource(&FilterNameArray)
+								.Font(FCoreStyle::Get().GetFontStyle("SlateFileDialogs.Dialog"))
+								.OnSelectionChanged(this, &SSlateFileOpenDlg::OnFilterChanged)
+							]
+						]
+					]
+
+					+ SHorizontalBox::Slot()
+					.HAlign(HAlign_Fill)
+					.VAlign(VAlign_Fill)
+					.FillWidth(1.0f)
+					[
+						SNew(SSpacer)
+							.Size(FVector2D(1.0f, 1.0f))
+					]
+
+					+ SHorizontalBox::Slot()
+					.HAlign(HAlign_Center)
+					.VAlign(VAlign_Center)
+					.Padding(FMargin(0.0f, 0.0f, 20.0f, 0.0f))
+					.AutoWidth()
+					[
+						SNew(SButton)
+						.ContentPadding(FMargin(5.0f, 5.0f))
+						.OnClicked(this, &SSlateFileOpenDlg::OnAcceptCancelClick, FSlateFileDlgWindow::Accept)
+						.Text(AcceptText)
+					]
+
+					+ SHorizontalBox::Slot()
+					.HAlign(HAlign_Center)
+					.VAlign(VAlign_Center)
+					.Padding(FMargin(0.0f))
+					.AutoWidth()
+					[
+						SNew(SButton)
+						.ContentPadding(FMargin(5.0f, 5.0f))
+						.OnClicked(this, &SSlateFileOpenDlg::OnAcceptCancelClick, FSlateFileDlgWindow::Cancel)
+						.Text(LOCTEXT("SlateFileDialogsCancel","Cancel"))
+					]
+				]
+			]
+		];
+
+	SaveFilename = "";
+
+	bNeedsBuilding = true;
+	bRebuildDirPath = true;
+	bDirectoryHasChanged = false;
+	DirectoryWatcher = nullptr;
+
+#if ENABLE_DIRECTORY_WATCHER	
+	if (!FModuleManager::Get().IsModuleLoaded("DirectoryWatcher"))
+	{
+		FModuleManager::Get().LoadModule("DirectoryWatcher");
+	}
+	
+	FDirectoryWatcherModule &DirWatcherModule = FModuleManager::LoadModuleChecked<FDirectoryWatcherModule>(TEXT("DirectoryWatcher"));
+	DirectoryWatcher = DirWatcherModule.Get();
+#endif
+}
+END_SLATE_FUNCTION_BUILD_OPTIMIZATION
+
+
+void SSlateFileOpenDlg::BuildDirectoryPath()
+{
+	// Clean up path as needed. Fix slashes and convert to absolute path.
+	FString NormPath = CurrentPath.Get();
+	FPaths::NormalizeFilename(NormPath);
+	FPaths::RemoveDuplicateSlashes(NormPath);
+	FString AbsPath = FPaths::ConvertRelativePathToFull(NormPath);
+	TCHAR Temp[MAX_PATH_LENGTH];
+
+	DirectoryNodesArray.Empty();
+
+	FString BuiltPath;
+	if (PLATFORM_WINDOWS)
+	{
+		int32 Idx;
+
+		if (AbsPath.FindChar(TCHAR('/'), Idx))
+		{
+			BuiltPath = BuiltPath + TEXT("/") + AbsPath.Left(Idx);
+		}
+
+		FCString::Strcpy(Temp, AbsPath.Len(), &AbsPath[Idx+1]);
+
+		DirectoryNodesArray.Add(FDirNode(AbsPath.Left(Idx), nullptr));
+	}
+	else if (PLATFORM_LINUX)
+	{
+		// start with system base directory
+		FCString::Strcpy(Temp, AbsPath.Len(), *AbsPath);
+
+		BuiltPath = "/";
+		DirectoryNodesArray.Add(FDirNode(FString(TEXT("/")), nullptr));
+	}
+
+	// break path into tokens
+	TCHAR *ContextStr = nullptr;
+	TCHAR *DirNode = FCString::Strtok(Temp, TEXT("/"), &ContextStr);
+
+	while (DirNode)
+	{
+		FString Label = DirNode;
+		DirectoryNodesArray.Add(FDirNode(Label, nullptr));
+		BuiltPath = BuiltPath + Label + TEXT("/");
+	
+		DirNode = FCString::Strtok(nullptr, TEXT("/"), &ContextStr);
+	}
+	
+	RefreshCrumbs();
+}
+
+
+void SSlateFileOpenDlg::RefreshCrumbs()
+{
+	// refresh crumb list
+	if (PathBreadcrumbTrail.IsValid())
+	{
+		PathBreadcrumbTrail->ClearCrumbs();
+
+		FString BuiltPath;
+		if (PLATFORM_WINDOWS)
+		{
+			PathBreadcrumbTrail->PushCrumb(LOCTEXT("SlateFileDialogsSystem", "System"), FString("SYSTEM"));
+
+			for (int32 i = 0; i < DirectoryNodesArray.Num(); i++)
+			{
+				BuiltPath = BuiltPath + DirectoryNodesArray[i].Label + TEXT("/");
+				PathBreadcrumbTrail->PushCrumb(FText::FromString(DirectoryNodesArray[i].Label), BuiltPath);
+			}
+		}
+		else if (PLATFORM_LINUX)
+		{
+			BuiltPath = "/";
+			PathBreadcrumbTrail->PushCrumb(FText::FromString(BuiltPath), BuiltPath);
+	
+			for (int32 i = 1; i < DirectoryNodesArray.Num(); i++)
+			{
+				BuiltPath = BuiltPath + DirectoryNodesArray[i].Label + TEXT("/");
+				PathBreadcrumbTrail->PushCrumb(FText::FromString(DirectoryNodesArray[i].Label), BuiltPath);
+			}
+		}
+	}
+}
+
+
+void SSlateFileOpenDlg::OnPathClicked(const FString & NewPath)
+{
+	if (NewPath.Compare("SYSTEM") == 0)
+	{
+		// Ignore clicks on the virtual root. (Only happens for Windows systems.)
+		return;
+	}
+
+	// set new current path and flag that we need to update directory display
+	CurrentPath = NewPath;
+	bRebuildDirPath = false;
+	bNeedsBuilding = true;
+
+	RefreshCrumbs();
+}
+
+
+void SSlateFileOpenDlg::OnPathMenuItemClicked( FString ClickedPath )
+{	
+	CurrentPath = ClickedPath;
+	bRebuildDirPath = true;
+	bNeedsBuilding = true;
+
+	RefreshCrumbs();
+}
+
+
+TSharedPtr<SWidget> SSlateFileOpenDlg::OnGetCrumbDelimiterContent(const FString& CrumbData) const
+{
+	TSharedPtr<SWidget> Widget = SNullWidget::NullWidget;
+	TArray<FString> SubDirs;
+
+	IFileManager& FileManager = IFileManager::Get();	
+	FSlateFileDialogDirVisitor DirVisitor(&SubDirs);
+
+	if (PLATFORM_WINDOWS)
+	{
+		if (CrumbData.Compare("SYSTEM") == 0)
+		{
+			// Windows doesn't have a root file system. So we need to provide a way to select system drives.
+			// This is done by creating a virtual root using 'System' as the top node.
+			int32 DrivesMask =
+#if PLATFORM_WINDOWS
+					(int32)GetLogicalDrives()
+#else
+					0
+#endif // PLATFORM_WINDOWS
+					;
+
+			FMenuBuilder MenuBuilder(true, NULL);
+			const TCHAR *DriveLetters = TEXT("ABCDEFGHIJKLMNOPQRSTUVWXYZ");
+			FString Drive = TEXT("A:");
+
+			for (int32 i = 0; i < 26; i++)
+			{
+				if (DrivesMask & 0x01)
+				{
+					Drive[0] = DriveLetters[i];
+
+					MenuBuilder.AddMenuEntry(
+						FText::FromString(Drive),
+						FText::GetEmpty(),
+						FSlateIcon(),
+						FUIAction(FExecuteAction::CreateSP(this, &SSlateFileOpenDlg::OnPathMenuItemClicked, Drive + TEXT("/"))));
+				}
+
+				DrivesMask >>= 1;
+			}
+
+			return SNew(SVerticalBox)
+				+ SVerticalBox::Slot()
+				.MaxHeight(400.0f)
+				[
+					MenuBuilder.MakeWidget()
+				];
+		}
+	}
+
+	FileManager.IterateDirectory(*CrumbData, DirVisitor);
+
+	if (SubDirs.Num() > 0)
+	{
+		SubDirs.Sort();
+
+		FMenuBuilder MenuBuilder(  true, NULL );
+
+		for (int32 i = 0; i < SubDirs.Num(); i++)
+		{
+			const FString& SubDir = SubDirs[i];
+
+			MenuBuilder.AddMenuEntry(
+				FText::FromString(SubDir),
+				FText::GetEmpty(),
+				FSlateIcon(),
+				FUIAction(FExecuteAction::CreateSP(this, &SSlateFileOpenDlg::OnPathMenuItemClicked, CrumbData + SubDir + TEXT("/"))));
+		}
+
+		Widget = 
+			SNew( SVerticalBox )
+			+SVerticalBox::Slot()
+			.MaxHeight(400.0f)
+			[
+				MenuBuilder.MakeWidget()
+			];
+	}
+
+	return Widget;
+}
+
+
+FReply SSlateFileOpenDlg::OnQuickLinkClick(FSlateFileDlgWindow::EResult ButtonID)
+{
+	if (ButtonID == FSlateFileDlgWindow::Project)
+	{
+		// Taken from DesktopPlatform. We have to do this to avoid a circular dependency.
+		const FString DefaultProjectSubFolder =TEXT("Unreal Projects");
+		CurrentPath = FString(FPlatformProcess::UserDir()) + DefaultProjectSubFolder + TEXT("/");
+	}
+
+	if (ButtonID == FSlateFileDlgWindow::Engine)
+	{
+		CurrentPath = FPaths::EngineDir();
+	}
+
+	bNeedsBuilding = true;
+	bRebuildDirPath = true;
+
+	return FReply::Handled();
+}
+
+
+FReply SSlateFileOpenDlg::OnAcceptCancelClick(FSlateFileDlgWindow::EResult ButtonID)
+{
+	if (ButtonID == FSlateFileDlgWindow::Accept)
+	{
+		if (OutNames.Get() != nullptr)
+		{
+			OutNames.Get()->Empty();
+
+			if (bSaveFile.Get())
+			{
+				FString Path = CurrentPath.Get() + SaveFilename;
+				OutNames.Get()->Add(Path);
+			}
+			else
+			{
+				TArray<TSharedPtr<FFileEntry>> SelectedItems = ListView->GetSelectedItems();
+
+				if (bDirectoriesOnly.Get())
+				{
+					if (SelectedItems.Num() > 0)
+					{
+						FString Path = CurrentPath.Get() + SelectedItems[0]->Label;
+						OutNames.Get()->Add(Path);
+					}
+				}
+				else
+				{
+					for (int32 i = 0; i < SelectedItems.Num(); i++)
+					{
+							FString Path = CurrentPath.Get() + SelectedItems[i]->Label;
+							OutNames.Get()->Add(Path);
+					}
+
+					if (OutFilterIndex.Get() != nullptr)
+					{
+						*(OutFilterIndex.Get()) = FilterIndex;
+					}
+				}
+			}
+		}
+	}
+	else
+	{
+		FString Path = "";
+		OutNames.Get()->Add(Path);
+	}
+
+	UserResponse = ButtonID;
+
+	ParentWindow.Get().Pin()->RequestDestroyWindow();
+
+	return FReply::Handled();
+}
+
+
+FReply SSlateFileOpenDlg::OnDirSublevelClick(int32 Level)
+{
+	DirectoryNodesArray[DirNodeIndex].TextBlock->SetFont(FCoreStyle::Get().GetFontStyle("SlateFileDialogs.Dialog"));
+
+	FString NewPath = TEXT("/");
+
+	for (int32 i = 1; i <= Level; i++)
+	{
+		NewPath += DirectoryNodesArray[i].Label + TEXT("/");
+	}
+
+	CurrentPath = NewPath;
+	bRebuildDirPath = false;
+	bNeedsBuilding = true;
+	
+	DirNodeIndex = Level;
+	DirectoryNodesArray[DirNodeIndex].TextBlock->SetFont(FCoreStyle::Get().GetFontStyle("SlateFileDialogs.DialogBold"));
+
+	return FReply::Handled();
+}
+
+
+
+void SSlateFileOpenDlg::Tick(const FGeometry &AllottedGeometry, const double InCurrentTime, const float InDeltaTime)
+{
+	SCompoundWidget::Tick(AllottedGeometry, InCurrentTime, InDeltaTime);
+
+	if (DirectoryWatcher)
+	{
+		DirectoryWatcher->Tick(InDeltaTime);
+	}
+
+	if (bDirectoryHasChanged && !bNeedsBuilding)
+	{
+		ReadDir(true);
+		RebuildFileTable();
+		ListView->RequestListRefresh();
+		bDirectoryHasChanged = false;			
+	}		
+
+	if (bNeedsBuilding)
+	{
+		// quick-link buttons to directory sublevels
+		if (bRebuildDirPath)
+		{
+			BuildDirectoryPath();
+		}
+
+		// Get directory contents and rebuild list
+		ParseFilters();
+		ReadDir();
+		RebuildFileTable();
+		ListView->RequestListRefresh();
+	}
+
+	bNeedsBuilding = false;
+	bRebuildDirPath = false;
+}
+
+
+void SSlateFileOpenDlg::ReadDir(bool bIsRefresh)
+{
+	if (DirectoryWatcher && RegisteredPath.Len() > 0 && !bIsRefresh)
+	{
+		DirectoryWatcher->UnregisterDirectoryChangedCallback_Handle(RegisteredPath, OnDialogDirectoryChangedDelegateHandle);
+		RegisteredPath = TEXT("");
+	}
+
+	IFileManager& FileManager = IFileManager::Get();
+
+	FilesArray.Empty();
+	FoldersArray.Empty();
+	TSharedPtr<FString> FilterList = nullptr;
+
+	if (FilterListArray.Num() > 0 && FilterIndex >= 0)
+	{
+		FilterList = FilterListArray[FilterIndex];
+	}
+
+	FSlateFileDialogVisitor DirVisitor(FilesArray, FoldersArray, FilterList);
+
+	FileManager.IterateDirectory(*CurrentPath.Get(), DirVisitor);
+
+	FilesArray.Sort(FFileEntry::ConstPredicate);
+	FoldersArray.Sort(FFileEntry::ConstPredicate);
+
+	if (DirectoryWatcher && !bIsRefresh)
+	{
+		DirectoryWatcher->RegisterDirectoryChangedCallback_Handle(CurrentPath.Get(),
+						IDirectoryWatcher::FDirectoryChanged::CreateRaw(this, &SSlateFileOpenDlg::OnDirectoryChanged),
+						OnDialogDirectoryChangedDelegateHandle);
+
+		RegisteredPath = CurrentPath.Get();
+	}
+}
+
+
+void SSlateFileOpenDlg::OnDirectoryChanged(const TArray <FFileChangeData> &FileChanges)
+{
+	bDirectoryHasChanged = true;
+}
+
+
+
+void SSlateFileOpenDlg::RebuildFileTable()
+{
+	LineItemArray.Empty();
+
+	// directory entries
+	for (int32 i = 0; i < FoldersArray.Num(); i++)
+	{
+		LineItemArray.Add(FoldersArray[i]);
+	}
+
+	// file entries
+	if (bDirectoriesOnly.Get() == false)
+	{
+		for (int32 i = 0; i < FilesArray.Num(); i++)
+		{
+			LineItemArray.Add(FilesArray[i]);
+		}
+	}
+}
+
+
+TSharedRef<ITableRow> SSlateFileOpenDlg::OnGenerateWidgetForList(TSharedPtr<FFileEntry> Item,
+		const TSharedRef<STableViewBase> &OwnerTable)
+{
+	return SNew(SSlateFileDialogRow, OwnerTable)
+			.DialogItem(Item);
+}
+
+
+void SSlateFileOpenDlg::OnItemDoubleClicked(TSharedPtr<FFileEntry> Item)
+{
+	if (Item->bIsDirectory)
+	{
+		CurrentPath = CurrentPath.Get() + Item->Label + TEXT("/");
+		bNeedsBuilding = true;
+		bRebuildDirPath = true;
+	}
+}
+
+
+void SSlateFileOpenDlg::OnFilterChanged(TSharedPtr<FString> NewValue, ESelectInfo::Type SelectInfo)
+{
+	for (int32 i = 0; i < FilterNameArray.Num(); i++)
+	{
+		if (FilterNameArray[i].Get()->Compare(*NewValue.Get(), ESearchCase::CaseSensitive) == 0)
+		{
+			FilterIndex = i;
+			break;
+		}
+	}
+
+	bNeedsBuilding = true;
+}
+
+
+void SSlateFileOpenDlg::SetDefaultFile(FString DefaultFile)
+{
+	SaveFilename = DefaultFile;
+	SaveFilenameEditBox->SetText(SaveFilename);
+}
+
+
+void SSlateFileOpenDlg::OnFileNameCommitted(const FText & InText, ETextCommit::Type InCommitType)
+{
+	// update edit box unless user choose to escape out
+	if (InCommitType != ETextCommit::OnCleared)
+	{
+		FString Extension;
+
+		SaveFilename = InText.ToString();
+
+		// get current filter extension
+		if (GetFilterExtension(Extension))
+		{
+			// append extension to filename if user left it off
+			if (!SaveFilename.EndsWith(Extension, ESearchCase::CaseSensitive))
+			{
+				SaveFilename = SaveFilename + Extension;
+			}
+		}
+
+		// remove any whitespace from start and/or end of filename
+		SaveFilename.Trim();
+		SaveFilename.TrimTrailing();
+
+		// update edit box content
+		SaveFilenameEditBox->SetText(SaveFilename);
+	}
+}
+
+
+void SSlateFileOpenDlg::OnItemSelected(TSharedPtr<FFileEntry> Item, ESelectInfo::Type SelectInfo)
+{
+	if (Item.IsValid())
+	{
+		if (!bDirectoriesOnly.Get())
+		{
+			TArray<TSharedPtr<FFileEntry>> SelectedItems = ListView->GetSelectedItems();
+			
+			for (int32 i = 0; i < SelectedItems.Num(); i++)
+			{
+				if (SelectedItems[i]->bIsDirectory)
+				{
+					ListView->SetItemSelection(SelectedItems[i], false, ESelectInfo::Direct);
+				}
+			}
+		}
+
+		if (bSaveFile.Get() && !Item->bIsDirectory)
+		{
+			SetDefaultFile(Item->Label);
+		}
+	}
+}
+
+
+void SSlateFileOpenDlg::ParseFilters()
+{
+	if (FilterCombo.IsValid() && FilterHBox.IsValid())
+	{
+		if (Filters.Get().Len() > 0)
+		{
+			if (FilterNameArray.Num() == 0)
+			{
+				TCHAR Temp[MAX_FILTER_LENGTH];
+				FCString::Strcpy(Temp, Filters.Get().Len(), *Filters.Get());
+
+				// break path into tokens
+				TCHAR *ContextStr = nullptr;
+
+				TCHAR *FilterDescription = FCString::Strtok(Temp, TEXT("|"), &ContextStr);
+				TCHAR *FilterList;
+
+				while (FilterDescription)
+				{
+					// filter wild cards
+					FilterList = FCString::Strtok(nullptr, TEXT("|"), &ContextStr);
+
+					FilterNameArray.Add(MakeShareable(new FString(FilterDescription)));
+					FilterListArray.Add(MakeShareable(new FString(FilterList)));
+
+					// next filter entry
+					FilterDescription = FCString::Strtok(nullptr, TEXT("|"), &ContextStr);
+				}
+			}
+
+			FilterCombo->SetSelectedItem(FilterNameArray[FilterIndex]);
+		}
+		else
+		{
+			FilterNameArray.Empty();
+			FilterHBox->SetVisibility(EVisibility::Hidden);
+		}
+	}
+}
+
+
+bool SSlateFileOpenDlg::GetFilterExtension(FString &OutString)
+{
+	// check to see if filters were given
+	if (Filters.Get().Len() == 0)
+	{
+		OutString = "";
+		return false;
+	}
+
+	// make a copy of filter string that we can modify
+	TCHAR Temp[MAX_FILTER_LENGTH];
+	FCString::Strcpy(Temp, FilterNameArray[FilterIndex]->Len(), *(*FilterNameArray[FilterIndex].Get()));
+
+	// find start of extension
+	TCHAR *FilterExt = FCString::Strchr(Temp, '.');
+
+	// strip any trailing junk
+	int32 i;
+	for (i = 0; i < FCString::Strlen(FilterExt); i++)
+	{
+		if (FilterExt[i] == ' ' || FilterExt[i] == ')' || FilterExt[i] == ';')
+		{
+			FilterExt[i] = 0;
+			break;
+		}
+	}
+
+	// store result and clean up
+	OutString = FString(FilterExt);
+
+	return (i > 0);
+}
+
+#undef LOCTEXT_NAMESPACE
diff --git a/Engine/Source/Developer/SlateFileDialogs/Private/SlateFileDlgWindow.h b/Engine/Source/Developer/SlateFileDialogs/Private/SlateFileDlgWindow.h
new file mode 100644
--- /dev/null
+++ b/Engine/Source/Developer/SlateFileDialogs/Private/SlateFileDlgWindow.h
@@ -0,0 +1,269 @@
+// Copyright 1998-2015 Epic Games, Inc. All Rights Reserved.
+
+#pragma once
+#include "SlateBasics.h"
+#include "STextComboBox.h"
+#include "SBreadcrumbTrail.h"
+#include "SInlineEditableTextBlock.h"
+
+#include "DirectoryWatcherModule.h"
+
+#define LOCTEXT_NAMESPACE "SlateFileDialogsNamespace"
+
+class SSlateFileOpenDlg;
+
+//-----------------------------------------------------------------------------
+struct FFileEntry
+{
+	FString	Label;
+	FString ModDate;
+	FString FileSize;
+	bool bIsSelected;
+	bool bIsDirectory;
+
+	FFileEntry() { }
+
+	FFileEntry(FString InLabel, FString InModDate, FString InFileSize, bool InIsDirectory)
+	{
+		Label = InLabel;
+		ModDate = FString(InModDate);
+		FileSize = FString(InFileSize);
+		bIsSelected = false;
+		bIsDirectory = InIsDirectory;
+	}
+
+	inline static bool ConstPredicate ( const TSharedPtr<FFileEntry> Entry1, const TSharedPtr<FFileEntry> Entry2)
+	{
+		return (Entry1->Label.Compare( Entry2->Label ) < 0);
+	}
+};
+
+typedef TSharedPtr<FFileEntry> SSlateFileDialogItemPtr;
+
+
+//-----------------------------------------------------------------------------
+class FSlateFileDlgWindow
+{
+public:
+	enum EResult
+	{
+		Cancel = 0,
+		Accept = 1,
+		Engine = 2,
+		Project = 3
+	};
+
+	bool OpenFileDialog(const void* ParentWindowHandle, const FString& DialogTitle, const FString& DefaultPath,
+		const FString& DefaultFile, const FString& FileTypes, uint32 Flags, TArray<FString>& OutFilenames, int32& OutFilterIndex);
+
+	bool OpenFileDialog(const void* ParentWindowHandle, const FString& DialogTitle, const FString& DefaultPath,
+		const FString& DefaultFile, const FString& FileTypes, uint32 Flags, TArray<FString>& OutFilenames);
+
+	bool OpenDirectoryDialog(const void* ParentWindowHandle, const FString& DialogTitle, const FString& DefaultPath,
+		FString& OutFoldername);
+
+	bool SaveFileDialog(const void* ParentWindowHandle, const FString& DialogTitle, const FString& DefaultPath,
+		const FString& DefaultFile, const FString& FileTypes, uint32 Flags, TArray<FString>& OutFilenames);
+
+private:	
+	TSharedPtr<class SSlateFileOpenDlg> DialogWidget;
+	FString CurrentDirectory;
+};
+
+
+
+//-----------------------------------------------------------------------------
+class SSlateFileOpenDlg : public SCompoundWidget
+{
+public:
+	struct FDirNode
+	{
+		FString	Label;
+		TSharedPtr<STextBlock> TextBlock;
+		TSharedPtr<SButton> Button;
+
+		FDirNode() { }
+
+		FDirNode(FString InLabel, TSharedPtr<STextBlock> InTextBlock)
+		{
+			Label = InLabel;
+			TextBlock = InTextBlock;
+		}
+	};
+
+	SLATE_BEGIN_ARGS(SSlateFileOpenDlg)
+		: _CurrentPath(TEXT("")),
+		_Filters(TEXT("")),
+		_bMultiSelectEnabled(false),
+		_WindowTitleText(TEXT("")),
+		_AcceptText(LOCTEXT("SlateDialogOpen","Open")),
+		_bDirectoriesOnly(false),
+		_bSaveFile(false),
+		_OutNames(nullptr),
+		_OutFilterIndex(nullptr),
+		_ParentWindow(nullptr)
+	{}
+
+	SLATE_ATTRIBUTE(FString, CurrentPath)
+	SLATE_ATTRIBUTE(FString, Filters)
+	SLATE_ATTRIBUTE(bool, bMultiSelectEnabled)
+	SLATE_ATTRIBUTE(FString, WindowTitleText)
+	SLATE_ATTRIBUTE(FText, AcceptText)
+	SLATE_ATTRIBUTE(bool, bDirectoriesOnly)
+	SLATE_ATTRIBUTE(bool, bSaveFile)
+	SLATE_ATTRIBUTE(TArray<FString> *, OutNames)
+	SLATE_ATTRIBUTE(int32 *, OutFilterIndex)
+	SLATE_ATTRIBUTE(TWeakPtr<SWindow>, ParentWindow)
+	SLATE_END_ARGS()
+
+	void Construct(const FArguments& InArgs);
+	FSlateFileDlgWindow::EResult GetResponse() { return UserResponse; }
+	void SetOutNames(TArray<FString> *Ptr) { OutNames = Ptr; }
+	void SetOutFilterIndex(int32 *FilterIndex) { OutFilterIndex = FilterIndex; }
+	void SetDefaultFile(FString DefaultFile);
+
+private:	
+	void OnPathClicked( const FString & CrumbData );
+	TSharedPtr<SWidget> OnGetCrumbDelimiterContent(const FString& CrumbData) const;
+	void RebuildFileTable();
+	void BuildDirectoryPath();
+	void ReadDir(bool bIsRefresh = false);
+	void RefreshCrumbs ( );
+	void OnPathMenuItemClicked( FString ClickedPath );
+	void OnItemSelected(TSharedPtr<FFileEntry> Item, ESelectInfo::Type SelectInfo);
+
+	void Tick(const FGeometry &AllottedGeometry, const double InCurrentTime, const float InDeltaTime);
+
+	TSharedRef<ITableRow> OnGenerateWidgetForList(TSharedPtr<FFileEntry> Item, const TSharedRef<STableViewBase> &OwnerTable);
+	void OnItemDoubleClicked(TSharedPtr<FFileEntry> Item);
+
+	void OnFilterChanged(TSharedPtr<FString> NewValue, ESelectInfo::Type SelectInfo);
+	FReply OnAcceptCancelClick(FSlateFileDlgWindow::EResult ButtonID);
+	FReply OnQuickLinkClick(FSlateFileDlgWindow::EResult ButtonID);
+	FReply OnDirSublevelClick(int32 Level);
+	void OnDirectoryChanged(const TArray <FFileChangeData> &FileChanges);	
+	void OnFileNameCommitted(const FText & InText, ETextCommit::Type InCommitType);
+	bool GetFilterExtension(FString &OutString);
+	void ParseFilters();
+
+	TArray< FDirNode > DirectoryNodesArray;
+	TArray<TSharedPtr<FFileEntry>> FoldersArray;
+	TArray<TSharedPtr<FFileEntry>> FilesArray;
+	TArray<TSharedPtr<FFileEntry>> LineItemArray;	
+
+	TSharedPtr<STextComboBox> FilterCombo;
+	TSharedPtr<SHorizontalBox> FilterHBox;
+	TSharedPtr<SInlineEditableTextBlock> SaveFilenameEditBox;
+	TSharedPtr<SBox> SaveFilenameSizeBox;
+	TSharedPtr<STextBlock> WindowTitle;
+	TSharedPtr<SListView<TSharedPtr<FFileEntry>>> ListView;
+	TSharedPtr<SBreadcrumbTrail<FString>> PathBreadcrumbTrail;
+	
+	TArray<TSharedPtr<FString>> FilterNameArray;
+	TArray<TSharedPtr<FString>> FilterListArray;
+
+	int32 FilterIndex;
+
+	FSlateFileDlgWindow::EResult	 UserResponse;
+		
+	bool bNeedsBuilding;
+	bool bRebuildDirPath;
+	bool bDirectoryHasChanged;
+
+	int32 DirNodeIndex;
+	FString SaveFilename;
+		
+	TAttribute<TWeakPtr<SWindow>> ParentWindow;
+	TAttribute<FString> CurrentPath;
+	TAttribute<FString> Filters;
+	TAttribute<FString> WindowTitleText;
+	TAttribute<bool> bMultiSelectEnabled;
+	TAttribute<TArray<FString> *> OutNames;
+	TAttribute<int32 *> OutFilterIndex;
+	TAttribute<bool> bDirectoriesOnly;
+	TAttribute<bool> bSaveFile;
+	TAttribute<FText> AcceptText;
+	
+	IDirectoryWatcher *DirectoryWatcher;
+	FDelegateHandle OnDialogDirectoryChangedDelegateHandle;
+	FString RegisteredPath;
+};
+
+
+
+//-----------------------------------------------------------------------------
+class SSlateFileDialogRow : public SMultiColumnTableRow<SSlateFileDialogItemPtr>
+{
+public:
+	
+	SLATE_BEGIN_ARGS(SSlateFileDialogRow) { }
+	SLATE_ARGUMENT(SSlateFileDialogItemPtr, DialogItem)
+	SLATE_END_ARGS()
+
+
+	void Construct(const FArguments& InArgs, const TSharedRef<STableViewBase> &InOwnerTableView)
+	{
+		check(InArgs._DialogItem.IsValid());
+		
+		DialogItem = InArgs._DialogItem;
+		
+		SMultiColumnTableRow<SSlateFileDialogItemPtr>::Construct(FSuperRowType::FArguments(), InOwnerTableView);
+	}
+
+	BEGIN_SLATE_FUNCTION_BUILD_OPTIMIZATION
+	virtual TSharedRef<SWidget> GenerateWidgetForColumn( const FName& ColumnName ) override
+	{
+		FSlateFontInfo ItemFont = FCoreStyle::Get().GetFontStyle("SlateFileDialogs.Dialog");
+		struct EVisibility FolderIconVisibility = DialogItem->bIsDirectory ? EVisibility::Visible : EVisibility::Hidden;
+
+		if (ColumnName == TEXT("Pathname"))
+		{
+			return SNew(SHorizontalBox)
+
+					+ SHorizontalBox::Slot()
+					.HAlign(HAlign_Left)
+					.VAlign(VAlign_Top)
+					.AutoWidth()
+					.Padding(FMargin(5.0f, 2.0f))
+					[
+						SNew(SImage)
+						.Image(FCoreStyle::Get().GetBrush("SlateFileDialogs.Folder16"))
+						.Visibility(FolderIconVisibility)
+					]
+
+					+ SHorizontalBox::Slot()
+					.HAlign(HAlign_Left)
+					.VAlign(VAlign_Top)
+					.AutoWidth()
+					.Padding(FMargin(5.0f, 0.0f, 0.0f, 0.0f))
+					[
+						SNew(STextBlock)
+						.Text(DialogItem->Label)
+						.Font(ItemFont)
+					];
+		}
+		else if (ColumnName == TEXT("ModDate"))
+		{
+			return SNew(STextBlock)
+					.Text(DialogItem->ModDate)
+					.Font(ItemFont);
+		}
+		else if (ColumnName == TEXT("FileSize"))
+		{
+			return SNew(STextBlock)
+					.Text(DialogItem->FileSize)
+					.Font(ItemFont);
+		}
+		else
+		{
+			return SNullWidget::NullWidget;
+		}
+	}	
+	END_SLATE_FUNCTION_BUILD_OPTIMIZATION
+
+	
+private:
+	SSlateFileDialogItemPtr DialogItem;
+};
+
+#undef LOCTEXT_NAMESPACE
diff --git a/Engine/Source/Developer/SlateFileDialogs/Public/ISlateFileDialogModule.h b/Engine/Source/Developer/SlateFileDialogs/Public/ISlateFileDialogModule.h
new file mode 100644
--- /dev/null
+++ b/Engine/Source/Developer/SlateFileDialogs/Public/ISlateFileDialogModule.h
@@ -0,0 +1,45 @@
+// Copyright 1998-2015 Epic Games, Inc. All Rights Reserved.
+
+#pragma once
+#include "Modules/ModuleInterface.h"
+
+/**
+ * Interface for slate file dialog modules.
+ */
+class ISlateFileDialogsModule : public IModuleInterface
+{
+public:
+	/**
+	 * Virtual destructor.
+	 */
+	virtual ~ISlateFileDialogsModule( ) { }
+	
+	virtual bool OpenFileDialog(const void* ParentWindowHandle, const FString& DialogTitle, const FString& DefaultPath,
+		const FString& DefaultFile, const FString& FileTypes, uint32 Flags, TArray<FString>& OutFilenames, int32& OutFilterIndex)
+	{
+		return false;
+	}
+	
+	virtual bool OpenFileDialog(const void* ParentWindowHandle, const FString& DialogTitle, const FString& DefaultPath,
+		const FString& DefaultFile, const FString& FileTypes, uint32 Flags, TArray<FString>& OutFilenames)
+	{
+		return false;
+	}
+	
+	virtual bool OpenDirectoryDialog(const void* ParentWindowHandle, const FString& DialogTitle, const FString& DefaultPath,
+		FString& OutFoldername)
+	{
+		return false;
+	}
+	
+	virtual bool SaveFileDialog(const void* ParentWindowHandle, const FString& DialogTitle, const FString& DefaultPath,
+		const FString& DefaultFile, const FString& FileTypes, uint32 Flags, TArray<FString>& OutFilenames)
+	{
+		return false;
+	}
+
+	virtual ISlateFileDialogsModule* Get() { return nullptr; }
+
+private:
+	ISlateFileDialogsModule *SlateFileDialog;
+};
diff --git a/Engine/Source/Developer/SlateFileDialogs/Public/SlateFileDialogs.h b/Engine/Source/Developer/SlateFileDialogs/Public/SlateFileDialogs.h
new file mode 100644
--- /dev/null
+++ b/Engine/Source/Developer/SlateFileDialogs/Public/SlateFileDialogs.h
@@ -0,0 +1,37 @@
+// Copyright 1998-2015 Epic Games, Inc. All Rights Reserved.
+
+#pragma once
+
+#include "Core.h"
+#include "ModuleInterface.h"
+#include "ISlateFileDialogModule.h"
+
+class FSlateFileDialogsModule : public ISlateFileDialogsModule
+{
+public:
+
+	// IModuleInterface interface
+
+	virtual void StartupModule() override;
+	virtual void ShutdownModule() override;
+
+public:
+	//ISlateFileDialogModule interface
+
+	bool OpenFileDialog(const void* ParentWindowHandle, const FString& DialogTitle, const FString& DefaultPath,
+		const FString& DefaultFile, const FString& FileTypes, uint32 Flags, TArray<FString>& OutFilenames, int32& OutFilterIndex);
+	
+	bool OpenFileDialog(const void* ParentWindowHandle, const FString& DialogTitle, const FString& DefaultPath,
+		const FString& DefaultFile, const FString& FileTypes, uint32 Flags, TArray<FString>& OutFilenames);
+	
+	bool OpenDirectoryDialog(const void* ParentWindowHandle, const FString& DialogTitle, const FString& DefaultPath,
+		FString& OutFoldername);
+
+	bool SaveFileDialog(const void* ParentWindowHandle, const FString& DialogTitle, const FString& DefaultPath,
+		const FString& DefaultFile, const FString& FileTypes, uint32 Flags, TArray<FString>& OutFilenames);
+	
+	ISlateFileDialogsModule* Get();
+
+private:
+	ISlateFileDialogsModule *SlateFileDialog;
+};
diff --git a/Engine/Source/Developer/SlateFileDialogs/SlateFileDialogs.Build.cs b/Engine/Source/Developer/SlateFileDialogs/SlateFileDialogs.Build.cs
new file mode 100644
--- /dev/null
+++ b/Engine/Source/Developer/SlateFileDialogs/SlateFileDialogs.Build.cs
@@ -0,0 +1,28 @@
+// Copyright 1998-2015 Epic Games, Inc. All Rights Reserved.
+
+using UnrealBuildTool;
+
+public class SlateFileDialogs : ModuleRules
+{
+    public SlateFileDialogs(TargetInfo Target)
+    {
+        PrivateDependencyModuleNames.AddRange(
+            new string[] {
+                "Core",
+                "InputCore",
+                "AppFramework",
+                "Slate",
+                "SlateCore",
+                "DirectoryWatcher",
+            }
+        );
+    
+        PrivateIncludePaths.AddRange(
+            new string[] {
+                "Developer/SlateFileDialogs/Private",  
+            }
+        );
+
+        PrivateIncludePathModuleNames.Add("TargetPlatform");
+    }
+}
diff --git a/Engine/Source/Runtime/SlateCore/Private/Styling/CoreStyle.cpp b/Engine/Source/Runtime/SlateCore/Private/Styling/CoreStyle.cpp
--- a/Engine/Source/Runtime/SlateCore/Private/Styling/CoreStyle.cpp
+++ b/Engine/Source/Runtime/SlateCore/Private/Styling/CoreStyle.cpp
@@ -1166,6 +1166,34 @@ TSharedRef<ISlateStyle> FCoreStyle::Create( const FName& InStyleSetName )
 		Style->Set("SyntaxHighlight.NodeAttributeValue", FTextBlockStyle(SmallMonospacedText).SetColorAndOpacity(FLinearColor(FColor(0xffb46100)))); // orange
 	}
 
+	// Slate File Dialogs
+	{
+		Style->Set("SlateFileDialogs.Dialog", TTF_FONT("Fonts/Roboto-Regular", 10));
+		Style->Set("SlateFileDialogs.DialogBold", TTF_FONT("Fonts/Roboto-Bold", 10));
+		Style->Set("SlateFileDialogs.DialogLarge", TTF_FONT("Fonts/Roboto-Bold", 16));
+		Style->Set("SlateFileDialogs.DirectoryItem", TTF_FONT("Fonts/Roboto-Bold", 11));
+		Style->Set( "SlateFileDialogs.GroupBorder", new BOX_BRUSH( "Common/GroupBorder", FMargin(4.0f/16.0f) ) );
+
+		Style->Set("SlateFileDialogs.Folder16", new IMAGE_BRUSH("Icons/icon_file_folder_16x", Icon16x16));
+		Style->Set("SlateFileDialogs.Folder24", new IMAGE_BRUSH("Icons/icon_file_folder_40x", Icon24x24));
+		Style->Set("SlateFileDialogs.WhiteBackground", new IMAGE_BRUSH("Common/Window/WindowWhite", Icon64x64));
+
+		Style->Set( "SlateFileDialogs.PathDelimiter", new IMAGE_BRUSH( "Common/SmallArrowRight", Icon10x10 ) );
+
+		Style->Set( "SlateFileDialogs.PathText", FTextBlockStyle(NormalText)
+			.SetFont( TTF_FONT( "Fonts/Roboto-Bold", 11 ) )
+			.SetColorAndOpacity( FLinearColor( 1.0f, 1.0f, 1.0f ) )
+			.SetHighlightColor( FLinearColor( 1.0f, 1.0f, 1.0f ) )
+			.SetShadowOffset( FVector2D( 1,1 ) )
+			.SetShadowColorAndOpacity( FLinearColor(0,0,0,0.9f) ) );
+
+		Style->Set("SlateFileDialogs.FlatButton", FButtonStyle(Button)
+			.SetNormal(FSlateNoResource())
+			.SetHovered(BOX_BRUSH("Common/FlatButton", 2.0f / 8.0f, SelectionColor))
+			.SetPressed(BOX_BRUSH("Common/FlatButton", 2.0f / 8.0f, SelectionColor_Pressed))
+			);
+	}
+
 	return Style;
 }
 
