SlateFileDialogs: fix improper concatenation (UE-17255).

diff --git a/Engine/Source/Developer/SlateFileDialogs/Private/SlateFileDlgWindow.cpp b/Engine/Source/Developer/SlateFileDialogs/Private/SlateFileDlgWindow.cpp
--- a/Engine/Source/Developer/SlateFileDialogs/Private/SlateFileDlgWindow.cpp
+++ b/Engine/Source/Developer/SlateFileDialogs/Private/SlateFileDlgWindow.cpp
@@ -682,6 +682,11 @@ void SSlateFileOpenDlg::Construct(const FArguments& InArgs)
 	bDirectoryHasChanged = false;
 	DirectoryWatcher = nullptr;
 
+	if (CurrentPath.Get().Len() > 0 && !CurrentPath.Get().EndsWith("/"))
+	{
+		CurrentPath = CurrentPath.Get() + TEXT("/");
+	}
+ 
 #if ENABLE_DIRECTORY_WATCHER	
 	if (!FModuleManager::Get().IsModuleLoaded("DirectoryWatcher"))
 	{
