SlateFileDialogs: updated from 6/17 code drop.

- Uses separate style and not Core.
- A new button to create the directory.

diff --git a/Engine/Source/Developer/SlateFileDialogs/Private/SlateFileDialogsModule.cpp b/Engine/Source/Developer/SlateFileDialogs/Private/SlateFileDialogsModule.cpp
--- a/Engine/Source/Developer/SlateFileDialogs/Private/SlateFileDialogsModule.cpp
+++ b/Engine/Source/Developer/SlateFileDialogs/Private/SlateFileDialogsModule.cpp
@@ -1,12 +1,17 @@
 // Copyright 1998-2015 Epic Games, Inc. All Rights Reserved.
 
 #include "SlateFileDialogsPrivatePCH.h"
+#include "SlateBasics.h"
 #include "SlateFileDialogs.h"
 #include "SlateFileDlgWindow.h"
+#include "SlateFileDialogsStyles.h"
 
 void FSlateFileDialogsModule::StartupModule()
 {
 	SlateFileDialog = new FSlateFileDialogsModule();
+
+	FileDialogStyle = new FSlateFileDialogsStyle;
+	FileDialogStyle->Initialize();
 }
 
 
@@ -14,6 +19,9 @@ void FSlateFileDialogsModule::ShutdownModule()
 {
 	if (SlateFileDialog != NULL)
 	{
+		FileDialogStyle->Shutdown();
+		delete FileDialogStyle;
+
 		delete SlateFileDialog;
 		SlateFileDialog = NULL;
 	}
@@ -23,7 +31,7 @@ void FSlateFileDialogsModule::ShutdownModule()
 bool FSlateFileDialogsModule::OpenFileDialog(const void* ParentWindowHandle, const FString& DialogTitle, const FString& DefaultPath,
 		const FString& DefaultFile, const FString& FileTypes, uint32 Flags, TArray<FString>& OutFilenames, int32& OutFilterIndex)
 {
-	FSlateFileDlgWindow dialog;
+	FSlateFileDlgWindow dialog(FileDialogStyle);
 	return dialog.OpenFileDialog(ParentWindowHandle, DialogTitle, DefaultPath, DefaultFile, FileTypes, Flags, OutFilenames, OutFilterIndex);
 }
 
@@ -31,14 +39,15 @@ bool FSlateFileDialogsModule::OpenFileDialog(const void* ParentWindowHandle, con
 bool FSlateFileDialogsModule::OpenFileDialog(const void* ParentWindowHandle, const FString& DialogTitle, const FString& DefaultPath,
 		const FString& DefaultFile, const FString& FileTypes, uint32 Flags, TArray<FString>& OutFilenames)
 {
-	FSlateFileDlgWindow dialog;
+	FSlateFileDlgWindow dialog(FileDialogStyle);
 	return dialog.OpenFileDialog(ParentWindowHandle, DialogTitle, DefaultPath, DefaultFile, FileTypes, Flags, OutFilenames);
 }
 
+
 bool FSlateFileDialogsModule::OpenDirectoryDialog(const void* ParentWindowHandle, const FString& DialogTitle, const FString& DefaultPath,
 		FString& OutFoldername)
 {
-	FSlateFileDlgWindow dialog;
+	FSlateFileDlgWindow dialog(FileDialogStyle);
 	return dialog.OpenDirectoryDialog(ParentWindowHandle, DialogTitle, DefaultPath, OutFoldername);
 }
 
@@ -46,7 +55,7 @@ bool FSlateFileDialogsModule::OpenDirectoryDialog(const void* ParentWindowHandle
 bool FSlateFileDialogsModule::SaveFileDialog(const void* ParentWindowHandle, const FString& DialogTitle, const FString& DefaultPath,
 	const FString& DefaultFile, const FString& FileTypes, uint32 Flags, TArray<FString>& OutFilenames)
 {
-	FSlateFileDlgWindow dialog;
+	FSlateFileDlgWindow dialog(FileDialogStyle);
 	return dialog.SaveFileDialog(ParentWindowHandle, DialogTitle, DefaultPath, DefaultFile, FileTypes, Flags, OutFilenames);
 }
 
diff --git a/Engine/Source/Developer/SlateFileDialogs/Private/SlateFileDialogsStyles.cpp b/Engine/Source/Developer/SlateFileDialogs/Private/SlateFileDialogsStyles.cpp
new file mode 100644
--- /dev/null
+++ b/Engine/Source/Developer/SlateFileDialogs/Private/SlateFileDialogsStyles.cpp
@@ -0,0 +1,110 @@
+// Copyright 1998-2015 Epic Games, Inc. All Rights Reserved.
+
+#include "SlateFileDialogsPrivatePCH.h"
+#include "SlateFileDialogsStyles.h"
+
+
+TSharedPtr< FSlateStyleSet > FSlateFileDialogsStyle::StyleInstance = nullptr;
+
+#define IMAGE_BRUSH( RelativePath, ... ) FSlateImageBrush(Style->RootToContentDir(RelativePath, TEXT(".png")), __VA_ARGS__)
+#define BOX_BRUSH( RelativePath, ... ) FSlateBoxBrush(Style->RootToContentDir(RelativePath, TEXT(".png")), __VA_ARGS__)
+#define BORDER_BRUSH( RelativePath, ... ) FSlateBorderBrush(Style->RootToContentDir(RelativePath, TEXT(".png")), __VA_ARGS__)
+#define TTF_FONT( RelativePath, ... ) FSlateFontInfo(Style->RootToContentDir(RelativePath, TEXT(".ttf")), __VA_ARGS__)
+
+
+void FSlateFileDialogsStyle::Initialize()
+{
+	if (!StyleInstance.IsValid())
+	{
+		StyleInstance = Create();
+		FSlateStyleRegistry::RegisterSlateStyle(*StyleInstance);
+	}
+}
+
+
+void FSlateFileDialogsStyle::Shutdown()
+{
+	FSlateStyleRegistry::UnRegisterSlateStyle(*StyleInstance);
+	ensure(StyleInstance.IsUnique());
+	StyleInstance.Reset();
+}
+
+
+FName FSlateFileDialogsStyle::GetStyleSetName()
+{
+	static FName StyleSetName(TEXT("SlateFileDialogsStyle"));
+	return StyleSetName;
+}
+
+
+const FSlateStyleSet *FSlateFileDialogsStyle::Get()
+{
+	return StyleInstance.Get();
+}
+
+
+TSharedPtr< FSlateStyleSet > FSlateFileDialogsStyle::Create()
+{
+	TSharedPtr< FSlateStyleSet > Style = MakeShareable(new FSlateStyleSet(FSlateFileDialogsStyle::GetStyleSetName()));
+	Style->SetContentRoot(FPaths::EngineContentDir());
+
+	const FVector2D Icon10x10(10.0f, 10.0f);
+	const FVector2D Icon16x16(16.0f, 16.0f);	
+	const FVector2D Icon24x24(24.0f, 24.0f);
+	const FVector2D Icon64x64(64.0f, 64.0f);
+
+	const FSlateColor InvertedForeground(FLinearColor(0.0f, 0.0f, 0.0f));
+	const FSlateColor SelectionColor(FLinearColor(0.701f, 0.225f, 0.003f));
+	const FSlateColor SelectionColor_Inactive(FLinearColor(0.25f, 0.25f, 0.25f));
+	const FSlateColor SelectionColor_Pressed(FLinearColor(0.701f, 0.225f, 0.003f));
+
+
+	// SButton defaults...
+	const FButtonStyle Button = FButtonStyle()
+		.SetNormal( BOX_BRUSH( "Slate/Common/Button", FVector2D(32,32), 8.0f/32.0f ) )
+		.SetHovered( BOX_BRUSH( "Slate/Common/Button_Hovered", FVector2D(32,32), 8.0f/32.0f ) )
+		.SetPressed( BOX_BRUSH( "Slate/Common/Button_Pressed", FVector2D(32,32), 8.0f/32.0f ) )
+		.SetDisabled( BOX_BRUSH( "Slate/Common/Button_Disabled", 8.0f/32.0f ) )
+		.SetNormalPadding( FMargin( 2,2,2,2 ) )
+		.SetPressedPadding( FMargin( 2,3,2,1 ) );
+	{
+		Style->Set( "Button", Button );
+
+		Style->Set( "InvertedForeground", InvertedForeground );
+	}
+
+	Style->Set("SlateFileDialogs.Dialog", TTF_FONT("Slate/Fonts/Roboto-Regular", 10));
+	Style->Set("SlateFileDialogs.DialogBold", TTF_FONT("Slate/Fonts/Roboto-Bold", 10));
+	Style->Set("SlateFileDialogs.DialogLarge", TTF_FONT("Slate/Fonts/Roboto-Bold", 16));
+	Style->Set("SlateFileDialogs.DirectoryItem", TTF_FONT("Slate/Fonts/Roboto-Bold", 11));
+	Style->Set( "SlateFileDialogs.GroupBorder", new BOX_BRUSH( "Slate/Common/GroupBorder", FMargin(4.0f/16.0f) ) );
+
+	Style->Set("SlateFileDialogs.Folder16", new IMAGE_BRUSH("SlateFileDialogs/Icons/icon_file_folder_16x", Icon16x16));
+	Style->Set("SlateFileDialogs.Folder24", new IMAGE_BRUSH("SlateFileDialogs/Icons/icon_file_folder_40x", Icon24x24));
+	Style->Set("SlateFileDialogs.NewFolder24", new IMAGE_BRUSH("SlateFileDialogs/Icons/icon_new_folder_40x", Icon24x24));		
+	Style->Set("SlateFileDialogs.WhiteBackground", new IMAGE_BRUSH("SlateFileDialogs/Common/Window/WindowWhite", Icon64x64));
+
+	Style->Set( "SlateFileDialogs.PathDelimiter", new IMAGE_BRUSH( "SlateFileDialogs/Common/SmallArrowRight", Icon10x10 ) );
+
+	Style->Set( "SlateFileDialogs.PathText", FTextBlockStyle()
+		.SetFont( TTF_FONT( "Slate/Fonts/Roboto-Bold", 11 ) )
+		.SetColorAndOpacity( FLinearColor( 1.0f, 1.0f, 1.0f ) )
+		.SetHighlightColor( FLinearColor( 1.0f, 1.0f, 1.0f ) )
+		.SetHighlightShape(BOX_BRUSH("Slate/Common/TextBlockHighlightShape", FMargin(3.f /8.f)))
+		.SetShadowOffset( FVector2D( 1,1 ) )
+		.SetShadowColorAndOpacity( FLinearColor(0,0,0,0.9f) ) );
+
+	Style->Set("SlateFileDialogs.FlatButton", FButtonStyle(Button)
+		.SetNormal(FSlateNoResource())
+		.SetHovered(BOX_BRUSH("SlateFileDialogs/Common/FlatButton", 2.0f / 8.0f, SelectionColor))
+		.SetPressed(BOX_BRUSH("SlateFileDialogs/Common/FlatButton", 2.0f / 8.0f, SelectionColor_Pressed))
+		);
+
+	return Style;	
+}
+
+
+#undef IMAGE_BRUSH
+#undef BOX_BRUSH
+#undef BORDER_BRUSH
+#undef TTF_FONT
diff --git a/Engine/Source/Developer/SlateFileDialogs/Private/SlateFileDlgWindow.cpp b/Engine/Source/Developer/SlateFileDialogs/Private/SlateFileDlgWindow.cpp
--- a/Engine/Source/Developer/SlateFileDialogs/Private/SlateFileDlgWindow.cpp
+++ b/Engine/Source/Developer/SlateFileDialogs/Private/SlateFileDlgWindow.cpp
@@ -216,6 +216,10 @@ private:
 
 //-----------------------------------------------------------------------------
 //-----------------------------------------------------------------------------
+FSlateFileDlgWindow::FSlateFileDlgWindow(FSlateFileDialogsStyle *InStyleSet)
+{
+	StyleSet = InStyleSet;
+}
 
 bool FSlateFileDlgWindow::OpenFileDialog(const void* ParentWindowHandle, const FString& DialogTitle, const FString& DefaultPath,
 		const FString& DefaultFile, const FString& FileTypes, uint32 Flags, TArray<FString>& OutFilenames, int32& OutFilterIndex)
@@ -235,7 +239,8 @@ bool FSlateFileDlgWindow::OpenFileDialog(const void* ParentWindowHandle, const F
 		.ParentWindow(ModalWindow)
 		.CurrentPath(DefaultPath)
 		.Filters(FileTypes)
-		.WindowTitleText(DialogTitle);
+		.WindowTitleText(DialogTitle)
+		.StyleSet(StyleSet);
 	
 	DialogWidget->SetOutNames(&OutFilenames);
 	DialogWidget->SetOutFilterIndex(&OutFilterIndex);
@@ -276,8 +281,9 @@ bool FSlateFileDlgWindow::OpenDirectoryDialog(const void* ParentWindowHandle, co
 		.bMultiSelectEnabled(false)
 		.ParentWindow(ModalWindow)
 		.bDirectoriesOnly(true)
-		.CurrentPath(DefaultPath)			
-		.WindowTitleText(DialogTitle);
+		.CurrentPath(DefaultPath)
+		.WindowTitleText(DialogTitle)
+		.StyleSet(StyleSet);
 
 	DialogWidget->SetOutNames(&TempOut);
 	DialogWidget->SetOutFilterIndex(&DummyIndex);
@@ -317,8 +323,9 @@ bool FSlateFileDlgWindow::SaveFileDialog(const void* ParentWindowHandle, const F
 		.bSaveFile(true)
 		.AcceptText(LOCTEXT("SlateFileDialogsSave","Save"))
 		.CurrentPath(DefaultPath)
-		.Filters(FileTypes)			
-		.WindowTitleText(DialogTitle);
+		.Filters(FileTypes)
+		.WindowTitleText(DialogTitle)
+		.StyleSet(StyleSet);
 
 	DialogWidget->SetOutNames(&OutFilenames);
 	DialogWidget->SetOutFilterIndex(&DummyIndex);
@@ -348,6 +355,7 @@ void SSlateFileOpenDlg::Construct(const FArguments& InArgs)
 	OutFilterIndex = InArgs._OutFilterIndex;
 	UserResponse = FSlateFileDlgWindow::Cancel;
 	ParentWindow = InArgs._ParentWindow;
+	StyleSet = InArgs._StyleSet;
 	AcceptText = InArgs._AcceptText;
 	DirNodeIndex = -1;
 	FilterIndex = 0;
@@ -361,7 +369,7 @@ void SSlateFileOpenDlg::Construct(const FArguments& InArgs)
 			.HAlign(HAlign_Fill)
 			.VAlign(VAlign_Fill)
 			.Padding(FMargin(20.0f,20.0f))
-			.BorderImage(FCoreStyle::Get().GetBrush("SlateFileDialogs.GroupBorder"))
+			.BorderImage(StyleSet.Get()->GetBrush("SlateFileDialogs.GroupBorder"))
 			[
 				SNew(SVerticalBox)
 
@@ -373,7 +381,7 @@ void SSlateFileOpenDlg::Construct(const FArguments& InArgs)
 				[
 					SAssignNew(WindowTitle, STextBlock)
 					.Text(WindowTitleText)
-					.Font(FCoreStyle::Get().GetFontStyle("SlateFileDialogs.DialogLarge"))
+					.Font(StyleSet.Get()->GetFontStyle("SlateFileDialogs.DialogLarge"))
 					.Justification(ETextJustify::Center)
 				]
 
@@ -381,13 +389,13 @@ void SSlateFileOpenDlg::Construct(const FArguments& InArgs)
 				.HAlign(HAlign_Fill)
 				.VAlign(VAlign_Fill)
 				.AutoHeight()
-				.Padding(FMargin(0.0f, 0.0f, 0.0f, 20.0f))
+				.Padding(FMargin(0.0f, 0.0f, 0.0f, 10.0f))
 				[
 					SAssignNew(PathBreadcrumbTrail, SBreadcrumbTrail<FString>)
 					.ButtonContentPadding(FMargin(2.0f, 2.0f))
-					.ButtonStyle(FCoreStyle::Get(), "SlateFileDialogs.FlatButton")
-					.DelimiterImage(FCoreStyle::Get().GetBrush("SlateFileDialogs.PathDelimiter"))
-					.TextStyle(FCoreStyle::Get(), "SlateFileDialogs.PathText")
+					.ButtonStyle(StyleSet.Get()->Get(), "SlateFileDialogs.FlatButton")
+					.DelimiterImage(StyleSet.Get()->GetBrush("SlateFileDialogs.PathDelimiter"))
+					.TextStyle(StyleSet.Get()->Get(), "SlateFileDialogs.PathText")
 					.ShowLeadingDelimiter(false)
 					.InvertTextColorOnHover(false)
 					.OnCrumbClicked(this, &SSlateFileOpenDlg::OnPathClicked)
@@ -395,6 +403,89 @@ void SSlateFileOpenDlg::Construct(const FArguments& InArgs)
 					.AddMetaData<FTagMetaData>(FTagMetaData(TEXT("ContentBrowserPath")))
 				]
 
+				+ SVerticalBox::Slot() // new directory
+				.HAlign(HAlign_Fill)
+				.VAlign(VAlign_Fill)
+				.AutoHeight()
+				.Padding(FMargin(0.0f, 0.0f, 0.0f, 10.0f))
+				[
+					SNew(SHorizontalBox)
+
+					+ SHorizontalBox::Slot()
+					.Padding(FMargin(0.0f))
+					.AutoWidth()
+					.HAlign(HAlign_Fill)
+					.VAlign(VAlign_Fill)
+					[
+						SNew(SButton)
+						.HAlign(HAlign_Center)
+						.VAlign(VAlign_Center)
+						.OnClicked(this, &SSlateFileOpenDlg::OnNewDirectoryClick)
+						.ContentPadding(FMargin(0.0f))
+						[
+							SNew(SImage)
+							.Image(StyleSet.Get()->GetBrush("SlateFileDialogs.NewFolder24"))
+						]
+					]
+
+					+ SHorizontalBox::Slot()
+					.Padding(FMargin(20.0f, 0.0f, 0.0f, 0.0f))
+					.AutoWidth()
+					.HAlign(HAlign_Fill)
+					.VAlign(VAlign_Fill)
+					[
+						SAssignNew(NewDirectorySizeBox, SBox)
+						.Padding(FMargin(0.0f))
+						.HAlign(HAlign_Fill)
+						.VAlign(VAlign_Fill)
+						.WidthOverride(300.0f)
+						.Visibility(EVisibility::Hidden)
+						[
+							SNew(SBorder)
+							.HAlign(HAlign_Fill)
+							.VAlign(VAlign_Fill)
+							.Padding(FMargin(5.0f,0.0f))
+							.BorderBackgroundColor(FLinearColor(0.1f, 0.1f, 0.1f, 1.0f))
+							.BorderImage(StyleSet.Get()->GetBrush("SlateFileDialogs.WhiteBackground"))
+							[
+								SAssignNew(NewDirectoryEditBox, SInlineEditableTextBlock)
+								.Font(StyleSet.Get()->GetFontStyle("SlateFileDialogs.Dialog"))
+								.IsReadOnly(false)
+								.Text(FText::FromString(""))
+								.OnTextCommitted(this, &SSlateFileOpenDlg::OnNewDirectoryCommitted)
+								.OnVerifyTextChanged(this, &SSlateFileOpenDlg::OnNewDirectoryTextChanged)
+							]
+						]
+					]
+
+					+ SHorizontalBox::Slot()
+					.HAlign(HAlign_Center)
+					.VAlign(VAlign_Center)
+					.Padding(FMargin(20.0f, 0.0f, 0.0f, 0.0f))
+					.AutoWidth()
+					[
+						SAssignNew(NewDirCancelButton, SButton)
+						.ContentPadding(FMargin(5.0f, 5.0f))
+						.OnClicked(this, &SSlateFileOpenDlg::OnNewDirectoryAcceptCancelClick, FSlateFileDlgWindow::Cancel)
+						.Text(LOCTEXT("SlateFileDialogsCancel","Cancel"))
+						.Visibility(EVisibility::Hidden)
+					]					
+				]
+
+				+ SVerticalBox::Slot() // new directory
+				.HAlign(HAlign_Fill)
+				.VAlign(VAlign_Fill)
+				.AutoHeight()
+				.Padding(FMargin(0.0f, 0.0f, 0.0f, 10.0f))
+				[
+					SAssignNew(DirErrorMsg, STextBlock)
+					.Font(StyleSet.Get()->GetFontStyle("SlateFileDialogs.DialogBold"))
+					.Justification(ETextJustify::Left)
+					.ColorAndOpacity(FLinearColor::Yellow)
+					.Text(LOCTEXT("SlateFileDialogsDirError", "Unable to create directory!"))
+					.Visibility(EVisibility::Collapsed)
+				]
+
 				+ SVerticalBox::Slot()
 				.HAlign(HAlign_Fill)
 				.VAlign(VAlign_Fill)
@@ -423,16 +514,16 @@ void SSlateFileOpenDlg::Construct(const FArguments& InArgs)
 							.ContentPadding(FMargin(2.0f))
 							[
 								SNew(SHorizontalBox)
-			
+
 								+ SHorizontalBox::Slot()
 								.HAlign(HAlign_Center)
 								.VAlign(VAlign_Center)
 								.AutoWidth()
 								[
 									SNew(SImage)
-									.Image(FCoreStyle::Get().GetBrush("SlateFileDialogs.Folder24"))
+									.Image(StyleSet.Get()->GetBrush("SlateFileDialogs.Folder24"))
 								]
-			
+
 								+ SHorizontalBox::Slot()
 								.HAlign(HAlign_Center)
 								.VAlign(VAlign_Center)
@@ -440,7 +531,7 @@ void SSlateFileOpenDlg::Construct(const FArguments& InArgs)
 								[
 									SNew(STextBlock)
 									.Text(FString("Projects"))
-									.Font(FCoreStyle::Get().GetFontStyle("SlateFileDialogs.Dialog"))
+									.Font(StyleSet.Get()->GetFontStyle("SlateFileDialogs.Dialog"))
 									.Justification(ETextJustify::Left)
 								]
 							]
@@ -466,7 +557,7 @@ void SSlateFileOpenDlg::Construct(const FArguments& InArgs)
 								.AutoWidth()
 								[
 									SNew(SImage)
-									.Image(FCoreStyle::Get().GetBrush("SlateFileDialogs.Folder24"))
+									.Image(StyleSet.Get()->GetBrush("SlateFileDialogs.Folder24"))
 								]
 
 								+ SHorizontalBox::Slot()
@@ -476,7 +567,7 @@ void SSlateFileOpenDlg::Construct(const FArguments& InArgs)
 								[
 									SNew(STextBlock)
 									.Text(FString("Engine"))
-									.Font(FCoreStyle::Get().GetFontStyle("SlateFileDialogs.Dialog"))
+									.Font(StyleSet.Get()->GetFontStyle("SlateFileDialogs.Dialog"))
 									.Justification(ETextJustify::Left)
 								]
 							]
@@ -504,7 +595,7 @@ void SSlateFileOpenDlg::Construct(const FArguments& InArgs)
 						.VAlign(VAlign_Fill)
 						.Padding(FMargin(10.0f))
 						.BorderBackgroundColor(FLinearColor(0.10f, 0.10f, 0.10f, 1.0f))
-						.BorderImage(FCoreStyle::Get().GetBrush("SlateFileDialogs.WhiteBackground"))
+						.BorderImage(StyleSet.Get()->GetBrush("SlateFileDialogs.WhiteBackground"))
 						[
 							SAssignNew(ListView, SListView<TSharedPtr<FFileEntry>>) // file list scroll
 							.ListItemsSource(&LineItemArray)
@@ -516,11 +607,11 @@ void SSlateFileOpenDlg::Construct(const FArguments& InArgs)
 							(
 								SNew(SHeaderRow)
 								.Visibility(EVisibility::Visible)
-								
+
 								+ SHeaderRow::Column("Pathname")
 									.DefaultLabel(LOCTEXT("SlateFileDialogsNameHeader", "Name"))
 									.FillWidth(1.0f)
-								
+
 								+ SHeaderRow::Column("ModDate")
 									.DefaultLabel(LOCTEXT("SlateFileDialogsModDateHeader", "Date Modified"))
 									.FixedWidth(170.0f)
@@ -556,7 +647,7 @@ void SSlateFileOpenDlg::Construct(const FArguments& InArgs)
 						[
 							SNew(STextBlock)
 							.Text(FString("Save Filename:"))
-							.Font(FCoreStyle::Get().GetFontStyle("SlateFileDialogs.Dialog"))
+							.Font(StyleSet.Get()->GetFontStyle("SlateFileDialogs.Dialog"))
 							.Justification(ETextJustify::Left)
 						]
 
@@ -577,10 +668,10 @@ void SSlateFileOpenDlg::Construct(const FArguments& InArgs)
 								.VAlign(VAlign_Fill)
 								.Padding(FMargin(5.0f,0.0f))
 								.BorderBackgroundColor(FLinearColor(0.1f, 0.1f, 0.1f, 1.0f))
-								.BorderImage(FCoreStyle::Get().GetBrush("SlateFileDialogs.WhiteBackground"))
+								.BorderImage(StyleSet.Get()->GetBrush("SlateFileDialogs.WhiteBackground"))
 								[
 									SAssignNew(SaveFilenameEditBox, SInlineEditableTextBlock)
-									.Font(FCoreStyle::Get().GetFontStyle("SlateFileDialogs.Dialog"))
+									.Font(StyleSet.Get()->GetFontStyle("SlateFileDialogs.Dialog"))
 									.IsReadOnly(false)
 									.Text(FText::FromString(""))
 									.OnTextCommitted(this, &SSlateFileOpenDlg::OnFileNameCommitted)
@@ -614,7 +705,7 @@ void SSlateFileOpenDlg::Construct(const FArguments& InArgs)
 						[
 							SNew(STextBlock)
 							.Text(FString("Filter:"))
-							.Font(FCoreStyle::Get().GetFontStyle("SlateFileDialogs.Dialog"))
+							.Font(StyleSet.Get()->GetFontStyle("SlateFileDialogs.Dialog"))
 							.Justification(ETextJustify::Left)
 						]
 
@@ -633,7 +724,7 @@ void SSlateFileOpenDlg::Construct(const FArguments& InArgs)
 								.ContentPadding(FMargin(4.0f, 2.0f))
 								//.MaxListHeight(100.0f)
 								.OptionsSource(&FilterNameArray)
-								.Font(FCoreStyle::Get().GetFontStyle("SlateFileDialogs.Dialog"))
+								.Font(StyleSet.Get()->GetFontStyle("SlateFileDialogs.Dialog"))
 								.OnSelectionChanged(this, &SSlateFileOpenDlg::OnFilterChanged)
 							]
 						]
@@ -686,13 +777,13 @@ void SSlateFileOpenDlg::Construct(const FArguments& InArgs)
 	{
 		CurrentPath = CurrentPath.Get() + TEXT("/");
 	}
- 
+
 #if ENABLE_DIRECTORY_WATCHER	
 	if (!FModuleManager::Get().IsModuleLoaded("DirectoryWatcher"))
 	{
 		FModuleManager::Get().LoadModule("DirectoryWatcher");
 	}
-	
+
 	FDirectoryWatcherModule &DirWatcherModule = FModuleManager::LoadModuleChecked<FDirectoryWatcherModule>(TEXT("DirectoryWatcher"));
 	DirectoryWatcher = DirWatcherModule.Get();
 #endif
@@ -773,7 +864,7 @@ void SSlateFileOpenDlg::RefreshCrumbs()
 		{
 			BuiltPath = "/";
 			PathBreadcrumbTrail->PushCrumb(FText::FromString(BuiltPath), BuiltPath);
-	
+
 			for (int32 i = 1; i < DirectoryNodesArray.Num(); i++)
 			{
 				BuiltPath = BuiltPath + DirectoryNodesArray[i].Label + TEXT("/");
@@ -806,7 +897,7 @@ void SSlateFileOpenDlg::OnPathMenuItemClicked( FString ClickedPath )
 	CurrentPath = ClickedPath;
 	bRebuildDirPath = true;
 	bNeedsBuilding = true;
-
+	
 	RefreshCrumbs();
 }
 
@@ -972,7 +1063,7 @@ FReply SSlateFileOpenDlg::OnAcceptCancelClick(FSlateFileDlgWindow::EResult Butto
 
 FReply SSlateFileOpenDlg::OnDirSublevelClick(int32 Level)
 {
-	DirectoryNodesArray[DirNodeIndex].TextBlock->SetFont(FCoreStyle::Get().GetFontStyle("SlateFileDialogs.Dialog"));
+	DirectoryNodesArray[DirNodeIndex].TextBlock->SetFont(StyleSet.Get()->GetFontStyle("SlateFileDialogs.Dialog"));
 
 	FString NewPath = TEXT("/");
 
@@ -984,9 +1075,9 @@ FReply SSlateFileOpenDlg::OnDirSublevelClick(int32 Level)
 	CurrentPath = NewPath;
 	bRebuildDirPath = false;
 	bNeedsBuilding = true;
-	
+
 	DirNodeIndex = Level;
-	DirectoryNodesArray[DirNodeIndex].TextBlock->SetFont(FCoreStyle::Get().GetFontStyle("SlateFileDialogs.DialogBold"));
+	DirectoryNodesArray[DirNodeIndex].TextBlock->SetFont(StyleSet.Get()->GetFontStyle("SlateFileDialogs.DialogBold"));
 
 	return FReply::Handled();
 }
@@ -996,7 +1087,7 @@ FReply SSlateFileOpenDlg::OnDirSublevelClick(int32 Level)
 void SSlateFileOpenDlg::Tick(const FGeometry &AllottedGeometry, const double InCurrentTime, const float InDeltaTime)
 {
 	SCompoundWidget::Tick(AllottedGeometry, InCurrentTime, InDeltaTime);
-
+	
 	if (DirectoryWatcher)
 	{
 		DirectoryWatcher->Tick(InDeltaTime);
@@ -1037,7 +1128,7 @@ void SSlateFileOpenDlg::ReadDir(bool bIsRefresh)
 		DirectoryWatcher->UnregisterDirectoryChangedCallback_Handle(RegisteredPath, OnDialogDirectoryChangedDelegateHandle);
 		RegisteredPath = TEXT("");
 	}
-
+	
 	IFileManager& FileManager = IFileManager::Get();
 
 	FilesArray.Empty();
@@ -1052,7 +1143,7 @@ void SSlateFileOpenDlg::ReadDir(bool bIsRefresh)
 	FSlateFileDialogVisitor DirVisitor(FilesArray, FoldersArray, FilterList);
 
 	FileManager.IterateDirectory(*CurrentPath.Get(), DirVisitor);
-
+	
 	FilesArray.Sort(FFileEntry::ConstPredicate);
 	FoldersArray.Sort(FFileEntry::ConstPredicate);
 
@@ -1077,7 +1168,7 @@ void SSlateFileOpenDlg::OnDirectoryChanged(const TArray <FFileChangeData> &FileC
 void SSlateFileOpenDlg::RebuildFileTable()
 {
 	LineItemArray.Empty();
-
+	
 	// directory entries
 	for (int32 i = 0; i < FoldersArray.Num(); i++)
 	{
@@ -1099,7 +1190,8 @@ TSharedRef<ITableRow> SSlateFileOpenDlg::OnGenerateWidgetForList(TSharedPtr<FFil
 		const TSharedRef<STableViewBase> &OwnerTable)
 {
 	return SNew(SSlateFileDialogRow, OwnerTable)
-			.DialogItem(Item);
+			.DialogItem(Item)
+			.StyleSet(StyleSet.Get());
 }
 
 
@@ -1149,7 +1241,7 @@ void SSlateFileOpenDlg::OnFileNameCommitted(const FText & InText, ETextCommit::T
 		if (GetFilterExtension(Extension))
 		{
 			// append extension to filename if user left it off
-			if (!SaveFilename.EndsWith(Extension, ESearchCase::CaseSensitive))
+			if (!SaveFilename.EndsWith(Extension, ESearchCase::CaseSensitive) && Extension.Compare(".*") != 0)
 			{
 				SaveFilename = SaveFilename + Extension;
 			}
@@ -1172,7 +1264,7 @@ void SSlateFileOpenDlg::OnItemSelected(TSharedPtr<FFileEntry> Item, ESelectInfo:
 		if (!bDirectoriesOnly.Get())
 		{
 			TArray<TSharedPtr<FFileEntry>> SelectedItems = ListView->GetSelectedItems();
-			
+
 			for (int32 i = 0; i < SelectedItems.Num(); i++)
 			{
 				if (SelectedItems[i]->bIsDirectory)
@@ -1264,4 +1356,71 @@ bool SSlateFileOpenDlg::GetFilterExtension(FString &OutString)
 	return (i > 0);
 }
 
+
+void SSlateFileOpenDlg::OnNewDirectoryCommitted(const FText & InText, ETextCommit::Type InCommitType)
+{
+	if (InCommitType == ETextCommit::OnEnter)
+	{
+		OnNewDirectoryAcceptCancelClick(FSlateFileDlgWindow::Accept);	
+	}
+	else
+	{
+		OnNewDirectoryAcceptCancelClick(FSlateFileDlgWindow::Cancel);
+	}
+}
+
+
+FReply SSlateFileOpenDlg::OnNewDirectoryClick()
+{
+	NewDirectorySizeBox->SetVisibility(EVisibility::Visible);
+	NewDirCancelButton->SetVisibility(EVisibility::Visible);
+	NewDirectoryEditBox->SetText(FString(""));
+	
+	FSlateApplication::Get().SetKeyboardFocus(NewDirectoryEditBox);	
+	NewDirectoryEditBox->EnterEditingMode();
+	
+	DirErrorMsg->SetVisibility(EVisibility::Collapsed);
+	
+	return FReply::Handled().SetUserFocus(NewDirectoryEditBox.ToSharedRef(), EFocusCause::SetDirectly);
+}
+
+
+bool SSlateFileOpenDlg::OnNewDirectoryTextChanged(const FText &InText, FText &ErrorMsg)
+{
+	NewDirectoryName = InText.ToString();
+	return true;
+}
+
+
+FReply SSlateFileOpenDlg::OnNewDirectoryAcceptCancelClick(FSlateFileDlgWindow::EResult ButtonID)
+{
+	if (ButtonID == FSlateFileDlgWindow::Accept)
+	{
+		NewDirectoryName.Trim();
+		NewDirectoryName.TrimTrailing();
+
+		if (NewDirectoryName.Len() > 0)
+		{
+			IPlatformFile &PlatformFile = FPlatformFileManager::Get().GetPlatformFile();
+			FString DirPath = CurrentPath.Get() + NewDirectoryName;
+
+			if (!PlatformFile.CreateDirectory(*DirPath))
+			{
+				DirErrorMsg->SetVisibility(EVisibility::Visible);				
+				return FReply::Handled();
+			}
+
+			bDirectoryHasChanged = true;
+		}
+	}
+
+	NewDirectorySizeBox->SetVisibility(EVisibility::Hidden);
+	NewDirCancelButton->SetVisibility(EVisibility::Hidden);
+	DirErrorMsg->SetVisibility(EVisibility::Collapsed);
+
+	NewDirectoryEditBox->SetText(FString(""));
+
+	return FReply::Handled();
+}
+
 #undef LOCTEXT_NAMESPACE
diff --git a/Engine/Source/Developer/SlateFileDialogs/Private/SlateFileDlgWindow.h b/Engine/Source/Developer/SlateFileDialogs/Private/SlateFileDlgWindow.h
--- a/Engine/Source/Developer/SlateFileDialogs/Private/SlateFileDlgWindow.h
+++ b/Engine/Source/Developer/SlateFileDialogs/Private/SlateFileDlgWindow.h
@@ -5,7 +5,7 @@
 #include "STextComboBox.h"
 #include "SBreadcrumbTrail.h"
 #include "SInlineEditableTextBlock.h"
-
+#include "SlateFileDialogsStyles.h"
 #include "DirectoryWatcherModule.h"
 
 #define LOCTEXT_NAMESPACE "SlateFileDialogsNamespace"
@@ -53,6 +53,8 @@ public:
 		Project = 3
 	};
 
+	FSlateFileDlgWindow(FSlateFileDialogsStyle *InStyleSet);
+
 	bool OpenFileDialog(const void* ParentWindowHandle, const FString& DialogTitle, const FString& DefaultPath,
 		const FString& DefaultFile, const FString& FileTypes, uint32 Flags, TArray<FString>& OutFilenames, int32& OutFilterIndex);
 
@@ -68,6 +70,8 @@ public:
 private:	
 	TSharedPtr<class SSlateFileOpenDlg> DialogWidget;
 	FString CurrentDirectory;
+
+	FSlateFileDialogsStyle* StyleSet;
 };
 
 
@@ -101,7 +105,8 @@ public:
 		_bSaveFile(false),
 		_OutNames(nullptr),
 		_OutFilterIndex(nullptr),
-		_ParentWindow(nullptr)
+		_ParentWindow(nullptr),
+		_StyleSet(nullptr)
 	{}
 
 	SLATE_ATTRIBUTE(FString, CurrentPath)
@@ -114,12 +119,13 @@ public:
 	SLATE_ATTRIBUTE(TArray<FString> *, OutNames)
 	SLATE_ATTRIBUTE(int32 *, OutFilterIndex)
 	SLATE_ATTRIBUTE(TWeakPtr<SWindow>, ParentWindow)
+	SLATE_ATTRIBUTE(FSlateFileDialogsStyle *, StyleSet)
 	SLATE_END_ARGS()
 
 	void Construct(const FArguments& InArgs);
 	FSlateFileDlgWindow::EResult GetResponse() { return UserResponse; }
 	void SetOutNames(TArray<FString> *Ptr) { OutNames = Ptr; }
-	void SetOutFilterIndex(int32 *FilterIndex) { OutFilterIndex = FilterIndex; }
+	void SetOutFilterIndex(int32 *InOutFilterIndex) { OutFilterIndex = InOutFilterIndex; }
 	void SetDefaultFile(FString DefaultFile);
 
 private:	
@@ -143,6 +149,12 @@ private:
 	FReply OnDirSublevelClick(int32 Level);
 	void OnDirectoryChanged(const TArray <FFileChangeData> &FileChanges);	
 	void OnFileNameCommitted(const FText & InText, ETextCommit::Type InCommitType);
+
+	void OnNewDirectoryCommitted(const FText & InText, ETextCommit::Type InCommitType);
+	FReply OnNewDirectoryClick();
+	bool OnNewDirectoryTextChanged(const FText &InText, FText &ErrorMsg);
+	FReply OnNewDirectoryAcceptCancelClick(FSlateFileDlgWindow::EResult ButtonID);
+
 	bool GetFilterExtension(FString &OutString);
 	void ParseFilters();
 
@@ -154,25 +166,30 @@ private:
 	TSharedPtr<STextComboBox> FilterCombo;
 	TSharedPtr<SHorizontalBox> FilterHBox;
 	TSharedPtr<SInlineEditableTextBlock> SaveFilenameEditBox;
+	TSharedPtr<SInlineEditableTextBlock> NewDirectoryEditBox;
 	TSharedPtr<SBox> SaveFilenameSizeBox;
 	TSharedPtr<STextBlock> WindowTitle;
 	TSharedPtr<SListView<TSharedPtr<FFileEntry>>> ListView;
 	TSharedPtr<SBreadcrumbTrail<FString>> PathBreadcrumbTrail;
-	
+
+	TSharedPtr<SButton> NewDirCancelButton;
+	TSharedPtr<SBox> NewDirectorySizeBox;
+	TSharedPtr<STextBlock> DirErrorMsg;
+
 	TArray<TSharedPtr<FString>> FilterNameArray;
 	TArray<TSharedPtr<FString>> FilterListArray;
 
 	int32 FilterIndex;
 
 	FSlateFileDlgWindow::EResult	 UserResponse;
-		
+
 	bool bNeedsBuilding;
 	bool bRebuildDirPath;
 	bool bDirectoryHasChanged;
 
 	int32 DirNodeIndex;
 	FString SaveFilename;
-		
+
 	TAttribute<TWeakPtr<SWindow>> ParentWindow;
 	TAttribute<FString> CurrentPath;
 	TAttribute<FString> Filters;
@@ -183,10 +200,12 @@ private:
 	TAttribute<bool> bDirectoriesOnly;
 	TAttribute<bool> bSaveFile;
 	TAttribute<FText> AcceptText;
-	
+	TAttribute<FSlateFileDialogsStyle *> StyleSet;
+
 	IDirectoryWatcher *DirectoryWatcher;
 	FDelegateHandle OnDialogDirectoryChangedDelegateHandle;
 	FString RegisteredPath;
+	FString NewDirectoryName;
 };
 
 
@@ -198,6 +217,7 @@ public:
 	
 	SLATE_BEGIN_ARGS(SSlateFileDialogRow) { }
 	SLATE_ARGUMENT(SSlateFileDialogItemPtr, DialogItem)
+	SLATE_ATTRIBUTE(FSlateFileDialogsStyle *, StyleSet)
 	SLATE_END_ARGS()
 
 
@@ -206,6 +226,7 @@ public:
 		check(InArgs._DialogItem.IsValid());
 		
 		DialogItem = InArgs._DialogItem;
+		StyleSet = InArgs._StyleSet;
 		
 		SMultiColumnTableRow<SSlateFileDialogItemPtr>::Construct(FSuperRowType::FArguments(), InOwnerTableView);
 	}
@@ -213,7 +234,7 @@ public:
 	BEGIN_SLATE_FUNCTION_BUILD_OPTIMIZATION
 	virtual TSharedRef<SWidget> GenerateWidgetForColumn( const FName& ColumnName ) override
 	{
-		FSlateFontInfo ItemFont = FCoreStyle::Get().GetFontStyle("SlateFileDialogs.Dialog");
+		FSlateFontInfo ItemFont = StyleSet.Get()->GetFontStyle("SlateFileDialogs.Dialog");
 		struct EVisibility FolderIconVisibility = DialogItem->bIsDirectory ? EVisibility::Visible : EVisibility::Hidden;
 
 		if (ColumnName == TEXT("Pathname"))
@@ -227,7 +248,7 @@ public:
 					.Padding(FMargin(5.0f, 2.0f))
 					[
 						SNew(SImage)
-						.Image(FCoreStyle::Get().GetBrush("SlateFileDialogs.Folder16"))
+						.Image(StyleSet.Get()->GetBrush("SlateFileDialogs.Folder16"))
 						.Visibility(FolderIconVisibility)
 					]
 
@@ -261,9 +282,10 @@ public:
 	}	
 	END_SLATE_FUNCTION_BUILD_OPTIMIZATION
 
-	
+
 private:
 	SSlateFileDialogItemPtr DialogItem;
+	TAttribute<FSlateFileDialogsStyle *> StyleSet;
 };
 
 #undef LOCTEXT_NAMESPACE
diff --git a/Engine/Source/Developer/SlateFileDialogs/Public/SlateFileDialogs.h b/Engine/Source/Developer/SlateFileDialogs/Public/SlateFileDialogs.h
--- a/Engine/Source/Developer/SlateFileDialogs/Public/SlateFileDialogs.h
+++ b/Engine/Source/Developer/SlateFileDialogs/Public/SlateFileDialogs.h
@@ -6,6 +6,8 @@
 #include "ModuleInterface.h"
 #include "ISlateFileDialogModule.h"
 
+class FSlateFileDialogsStyle;
+
 class FSlateFileDialogsModule : public ISlateFileDialogsModule
 {
 public:
@@ -20,18 +22,22 @@ public:
 
 	bool OpenFileDialog(const void* ParentWindowHandle, const FString& DialogTitle, const FString& DefaultPath,
 		const FString& DefaultFile, const FString& FileTypes, uint32 Flags, TArray<FString>& OutFilenames, int32& OutFilterIndex);
-	
+
 	bool OpenFileDialog(const void* ParentWindowHandle, const FString& DialogTitle, const FString& DefaultPath,
 		const FString& DefaultFile, const FString& FileTypes, uint32 Flags, TArray<FString>& OutFilenames);
-	
+
 	bool OpenDirectoryDialog(const void* ParentWindowHandle, const FString& DialogTitle, const FString& DefaultPath,
 		FString& OutFoldername);
 
 	bool SaveFileDialog(const void* ParentWindowHandle, const FString& DialogTitle, const FString& DefaultPath,
 		const FString& DefaultFile, const FString& FileTypes, uint32 Flags, TArray<FString>& OutFilenames);
-	
+
 	ISlateFileDialogsModule* Get();
 
+	FSlateFileDialogsStyle *GetFileDialogsStyle() { return FileDialogStyle; }
+
 private:
 	ISlateFileDialogsModule *SlateFileDialog;
+
+	FSlateFileDialogsStyle	*FileDialogStyle;
 };
diff --git a/Engine/Source/Developer/SlateFileDialogs/Public/SlateFileDialogsStyles.h b/Engine/Source/Developer/SlateFileDialogs/Public/SlateFileDialogsStyles.h
new file mode 100644
--- /dev/null
+++ b/Engine/Source/Developer/SlateFileDialogs/Public/SlateFileDialogsStyles.h
@@ -0,0 +1,28 @@
+// Copyright 1998-2015 Epic Games, Inc. All Rights Reserved.
+
+#pragma once
+
+#include "SlateBasics.h"
+
+class FSlateFileDialogsStyle
+{
+public:
+	static void Initialize();
+	static void Shutdown();
+	static const FSlateStyleSet *Get();
+	static FName GetStyleSetName();
+
+	static const FSlateBrush * GetBrush(FName PropertyName, const ANSICHAR* Specifier = NULL)
+	{
+		return StyleInstance->GetBrush(PropertyName, Specifier);
+	}
+
+	static FSlateFontInfo GetFontStyle(FName PropertyName, const ANSICHAR* Specifier = NULL)
+	{
+		return StyleInstance->GetFontStyle(PropertyName, Specifier);
+	}
+	
+private:
+	static TSharedPtr<class FSlateStyleSet> Create();
+	static TSharedPtr<class FSlateStyleSet> StyleInstance;
+};
diff --git a/Engine/Source/Runtime/SlateCore/Private/Styling/CoreStyle.cpp b/Engine/Source/Runtime/SlateCore/Private/Styling/CoreStyle.cpp
--- a/Engine/Source/Runtime/SlateCore/Private/Styling/CoreStyle.cpp
+++ b/Engine/Source/Runtime/SlateCore/Private/Styling/CoreStyle.cpp
@@ -1166,34 +1166,6 @@ TSharedRef<ISlateStyle> FCoreStyle::Create( const FName& InStyleSetName )
 		Style->Set("SyntaxHighlight.NodeAttributeValue", FTextBlockStyle(SmallMonospacedText).SetColorAndOpacity(FLinearColor(FColor(0xffb46100)))); // orange
 	}
 
-	// Slate File Dialogs
-	{
-		Style->Set("SlateFileDialogs.Dialog", TTF_FONT("Fonts/Roboto-Regular", 10));
-		Style->Set("SlateFileDialogs.DialogBold", TTF_FONT("Fonts/Roboto-Bold", 10));
-		Style->Set("SlateFileDialogs.DialogLarge", TTF_FONT("Fonts/Roboto-Bold", 16));
-		Style->Set("SlateFileDialogs.DirectoryItem", TTF_FONT("Fonts/Roboto-Bold", 11));
-		Style->Set( "SlateFileDialogs.GroupBorder", new BOX_BRUSH( "Common/GroupBorder", FMargin(4.0f/16.0f) ) );
-
-		Style->Set("SlateFileDialogs.Folder16", new IMAGE_BRUSH("Icons/icon_file_folder_16x", Icon16x16));
-		Style->Set("SlateFileDialogs.Folder24", new IMAGE_BRUSH("Icons/icon_file_folder_40x", Icon24x24));
-		Style->Set("SlateFileDialogs.WhiteBackground", new IMAGE_BRUSH("Common/Window/WindowWhite", Icon64x64));
-
-		Style->Set( "SlateFileDialogs.PathDelimiter", new IMAGE_BRUSH( "Common/SmallArrowRight", Icon10x10 ) );
-
-		Style->Set( "SlateFileDialogs.PathText", FTextBlockStyle(NormalText)
-			.SetFont( TTF_FONT( "Fonts/Roboto-Bold", 11 ) )
-			.SetColorAndOpacity( FLinearColor( 1.0f, 1.0f, 1.0f ) )
-			.SetHighlightColor( FLinearColor( 1.0f, 1.0f, 1.0f ) )
-			.SetShadowOffset( FVector2D( 1,1 ) )
-			.SetShadowColorAndOpacity( FLinearColor(0,0,0,0.9f) ) );
-
-		Style->Set("SlateFileDialogs.FlatButton", FButtonStyle(Button)
-			.SetNormal(FSlateNoResource())
-			.SetHovered(BOX_BRUSH("Common/FlatButton", 2.0f / 8.0f, SelectionColor))
-			.SetPressed(BOX_BRUSH("Common/FlatButton", 2.0f / 8.0f, SelectionColor_Pressed))
-			);
-	}
-
 	return Style;
 }
 
