SlateFileDialogs: disable directory watcher.

- Avoids locks up and other problems when browsing deep directory hierarchy.
- Temporary, a better solution is needed.

diff --git a/Engine/Source/Developer/SlateFileDialogs/Private/SlateFileDialogsPrivatePCH.h b/Engine/Source/Developer/SlateFileDialogs/Private/SlateFileDialogsPrivatePCH.h
--- a/Engine/Source/Developer/SlateFileDialogs/Private/SlateFileDialogsPrivatePCH.h
+++ b/Engine/Source/Developer/SlateFileDialogs/Private/SlateFileDialogsPrivatePCH.h
@@ -6,4 +6,4 @@
 #define MAX_FILTER_LENGTH 8192
 #define MAX_PATH_LENGTH 8192
 
-#define ENABLE_DIRECTORY_WATCHER 1
+#define ENABLE_DIRECTORY_WATCHER 0
