SlateFileDialogs: support opening files by double clicking (UE-17307).

diff --git a/Engine/Source/Developer/SlateFileDialogs/Private/SlateFileDlgWindow.cpp b/Engine/Source/Developer/SlateFileDialogs/Private/SlateFileDlgWindow.cpp
--- a/Engine/Source/Developer/SlateFileDialogs/Private/SlateFileDlgWindow.cpp
+++ b/Engine/Source/Developer/SlateFileDialogs/Private/SlateFileDlgWindow.cpp
@@ -1005,48 +1005,53 @@ FReply SSlateFileOpenDlg::OnQuickLinkClick(FSlateFileDlgWindow::EResult ButtonID
 	return FReply::Handled();
 }
 
-
-FReply SSlateFileOpenDlg::OnAcceptCancelClick(FSlateFileDlgWindow::EResult ButtonID)
+void SSlateFileOpenDlg::SetOutputFiles()
 {
-	if (ButtonID == FSlateFileDlgWindow::Accept)
+	if (OutNames.Get() != nullptr)
 	{
-		if (OutNames.Get() != nullptr)
+		OutNames.Get()->Empty();
+
+		if (bSaveFile.Get())
 		{
-			OutNames.Get()->Empty();
+			FString Path = CurrentPath.Get() + SaveFilename;
+			OutNames.Get()->Add(Path);
+		}
+		else
+		{
+			TArray<TSharedPtr<FFileEntry>> SelectedItems = ListView->GetSelectedItems();
 
-			if (bSaveFile.Get())
+			if (bDirectoriesOnly.Get())
 			{
-				FString Path = CurrentPath.Get() + SaveFilename;
-				OutNames.Get()->Add(Path);
+				if (SelectedItems.Num() > 0)
+				{
+					FString Path = CurrentPath.Get() + SelectedItems[0]->Label;
+					OutNames.Get()->Add(Path);
+				}
 			}
 			else
 			{
-				TArray<TSharedPtr<FFileEntry>> SelectedItems = ListView->GetSelectedItems();
-
-				if (bDirectoriesOnly.Get())
+				for (int32 i = 0; i < SelectedItems.Num(); i++)
 				{
-					if (SelectedItems.Num() > 0)
-					{
-						FString Path = CurrentPath.Get() + SelectedItems[0]->Label;
+						FString Path = CurrentPath.Get() + SelectedItems[i]->Label;
 						OutNames.Get()->Add(Path);
-					}
 				}
-				else
-				{
-					for (int32 i = 0; i < SelectedItems.Num(); i++)
-					{
-							FString Path = CurrentPath.Get() + SelectedItems[i]->Label;
-							OutNames.Get()->Add(Path);
-					}
 
-					if (OutFilterIndex.Get() != nullptr)
-					{
-						*(OutFilterIndex.Get()) = FilterIndex;
-					}
+				if (OutFilterIndex.Get() != nullptr)
+				{
+					*(OutFilterIndex.Get()) = FilterIndex;
 				}
 			}
 		}
 	}
+}
+
+
+FReply SSlateFileOpenDlg::OnAcceptCancelClick(FSlateFileDlgWindow::EResult ButtonID)
+{
+	if (ButtonID == FSlateFileDlgWindow::Accept)
+	{
+		SetOutputFiles();
+	}
 	else
 	{
 		FString Path = "";
@@ -1054,7 +1059,6 @@ FReply SSlateFileOpenDlg::OnAcceptCancelClick(FSlateFileDlgWindow::EResult Butto
 	}
 
 	UserResponse = ButtonID;
-
 	ParentWindow.Get().Pin()->RequestDestroyWindow();
 
 	return FReply::Handled();
@@ -1203,6 +1207,12 @@ void SSlateFileOpenDlg::OnItemDoubleClicked(TSharedPtr<FFileEntry> Item)
 		bNeedsBuilding = true;
 		bRebuildDirPath = true;
 	}
+	else
+	{
+		SetOutputFiles();
+		UserResponse = FSlateFileDlgWindow::Accept;
+		ParentWindow.Get().Pin()->RequestDestroyWindow();
+	}
 }
 
 
diff --git a/Engine/Source/Developer/SlateFileDialogs/Private/SlateFileDlgWindow.h b/Engine/Source/Developer/SlateFileDialogs/Private/SlateFileDlgWindow.h
--- a/Engine/Source/Developer/SlateFileDialogs/Private/SlateFileDlgWindow.h
+++ b/Engine/Source/Developer/SlateFileDialogs/Private/SlateFileDlgWindow.h
@@ -134,7 +134,7 @@ private:
 	void RebuildFileTable();
 	void BuildDirectoryPath();
 	void ReadDir(bool bIsRefresh = false);
-	void RefreshCrumbs ( );
+	void RefreshCrumbs();
 	void OnPathMenuItemClicked( FString ClickedPath );
 	void OnItemSelected(TSharedPtr<FFileEntry> Item, ESelectInfo::Type SelectInfo);
 
@@ -155,6 +155,9 @@ private:
 	bool OnNewDirectoryTextChanged(const FText &InText, FText &ErrorMsg);
 	FReply OnNewDirectoryAcceptCancelClick(FSlateFileDlgWindow::EResult ButtonID);
 
+	/** Collects the output files. */
+	void SetOutputFiles();
+
 	bool GetFilterExtension(FString &OutString);
 	void ParseFilters();
 
