SlateFileDialogs: default to current path when nothing is selected (UE-17299).

- This behavior only applies when opening a directory.

diff --git a/Engine/Source/Developer/SlateFileDialogs/Private/SlateFileDlgWindow.cpp b/Engine/Source/Developer/SlateFileDialogs/Private/SlateFileDlgWindow.cpp
--- a/Engine/Source/Developer/SlateFileDialogs/Private/SlateFileDlgWindow.cpp
+++ b/Engine/Source/Developer/SlateFileDialogs/Private/SlateFileDlgWindow.cpp
@@ -1027,6 +1027,11 @@ void SSlateFileOpenDlg::SetOutputFiles()
 					FString Path = CurrentPath.Get() + SelectedItems[0]->Label;
 					OutNames.Get()->Add(Path);
 				}
+				else
+				{
+					// select the current directory
+					OutNames.Get()->Add(CurrentPath.Get());
+				}
 			}
 			else
 			{
