SlateFileDialogs: use new DirectoryWatcher options to fix perf.

diff --git a/Engine/Source/Developer/SlateFileDialogs/Private/SlateFileDialogsPrivatePCH.h b/Engine/Source/Developer/SlateFileDialogs/Private/SlateFileDialogsPrivatePCH.h
--- a/Engine/Source/Developer/SlateFileDialogs/Private/SlateFileDialogsPrivatePCH.h
+++ b/Engine/Source/Developer/SlateFileDialogs/Private/SlateFileDialogsPrivatePCH.h
@@ -6,4 +6,4 @@
 #define MAX_FILTER_LENGTH 8192
 #define MAX_PATH_LENGTH 8192
 
-#define ENABLE_DIRECTORY_WATCHER 0
+#define ENABLE_DIRECTORY_WATCHER 1
diff --git a/Engine/Source/Developer/SlateFileDialogs/Private/SlateFileDlgWindow.cpp b/Engine/Source/Developer/SlateFileDialogs/Private/SlateFileDlgWindow.cpp
--- a/Engine/Source/Developer/SlateFileDialogs/Private/SlateFileDlgWindow.cpp
+++ b/Engine/Source/Developer/SlateFileDialogs/Private/SlateFileDlgWindow.cpp
@@ -1160,7 +1160,7 @@ void SSlateFileOpenDlg::ReadDir(bool bIsRefresh)
 	{
 		DirectoryWatcher->RegisterDirectoryChangedCallback_Handle(CurrentPath.Get(),
 						IDirectoryWatcher::FDirectoryChanged::CreateRaw(this, &SSlateFileOpenDlg::OnDirectoryChanged),
-						OnDialogDirectoryChangedDelegateHandle);
+						OnDialogDirectoryChangedDelegateHandle, IDirectoryWatcher::WatchOptions::IncludeDirectoryChanges | IDirectoryWatcher::WatchOptions::IgnoreChangesInSubtree);
 
 		RegisteredPath = CurrentPath.Get();
 	}
