SlateFileDialogs: fix issue with selecting a dir in the middle of the list.

diff --git a/Engine/Source/Developer/SlateFileDialogs/Private/SlateFileDlgWindow.cpp b/Engine/Source/Developer/SlateFileDialogs/Private/SlateFileDlgWindow.cpp
--- a/Engine/Source/Developer/SlateFileDialogs/Private/SlateFileDlgWindow.cpp
+++ b/Engine/Source/Developer/SlateFileDialogs/Private/SlateFileDlgWindow.cpp
@@ -885,7 +885,7 @@ void SSlateFileOpenDlg::OnPathClicked(const FString & NewPath)
 
 	// set new current path and flag that we need to update directory display
 	CurrentPath = NewPath;
-	bRebuildDirPath = false;
+	bRebuildDirPath = true;
 	bNeedsBuilding = true;
 
 	RefreshCrumbs();
