Updated SlateFileDialogs to no longer be dependent on SLATE_TEXT_ATTRIBUTE

SLATE_TEXT_ATTRIBUTE was the macro that allowed the FString -> FText pass-through support, and was deprecated in 4.8 in favour of using SLATE_ATTRIBUTE with an FText type. This is being removed for 4.10 and SlateFileDialogs was still relying on the deprecated behaviour.

This change removes its reliance on SLATE_TEXT_ATTRIBUTE, and also cleans up a load of SLATE_ATTRIBUTE parameters should have been using SLATE_ARGUMENT instead.

diff --git a/Engine/Source/Developer/SlateFileDialogs/Private/SlateFileDlgWindow.cpp b/Engine/Source/Developer/SlateFileDialogs/Private/SlateFileDlgWindow.cpp
--- a/Engine/Source/Developer/SlateFileDialogs/Private/SlateFileDlgWindow.cpp
+++ b/Engine/Source/Developer/SlateFileDialogs/Private/SlateFileDlgWindow.cpp
@@ -360,8 +360,8 @@ void SSlateFileOpenDlg::Construct(const FArguments& InArgs)
 	DirNodeIndex = -1;
 	FilterIndex = 0;
 
-	ESelectionMode::Type SelectMode = bMultiSelectEnabled.Get() ? ESelectionMode::Multi : ESelectionMode::Single;
-	struct EVisibility SaveFilenameVisibility = bSaveFile.Get() ? EVisibility::Visible : EVisibility::Collapsed;
+	ESelectionMode::Type SelectMode = bMultiSelectEnabled ? ESelectionMode::Multi : ESelectionMode::Single;
+	struct EVisibility SaveFilenameVisibility = bSaveFile ? EVisibility::Visible : EVisibility::Collapsed;
 
 	this->ChildSlot
 		[
@@ -369,7 +369,7 @@ void SSlateFileOpenDlg::Construct(const FArguments& InArgs)
 			.HAlign(HAlign_Fill)
 			.VAlign(VAlign_Fill)
 			.Padding(FMargin(20.0f,20.0f))
-			.BorderImage(StyleSet.Get()->GetBrush("SlateFileDialogs.GroupBorder"))
+			.BorderImage(StyleSet->GetBrush("SlateFileDialogs.GroupBorder"))
 			[
 				SNew(SVerticalBox)
 
@@ -380,8 +380,8 @@ void SSlateFileOpenDlg::Construct(const FArguments& InArgs)
 				.Padding(FMargin(0.0f, 0.0f, 0.0f, 20.0f))
 				[
 					SAssignNew(WindowTitle, STextBlock)
-					.Text(WindowTitleText)
-					.Font(StyleSet.Get()->GetFontStyle("SlateFileDialogs.DialogLarge"))
+					.Text(FText::FromString(WindowTitleText))
+					.Font(StyleSet->GetFontStyle("SlateFileDialogs.DialogLarge"))
 					.Justification(ETextJustify::Center)
 				]
 
@@ -393,9 +393,9 @@ void SSlateFileOpenDlg::Construct(const FArguments& InArgs)
 				[
 					SAssignNew(PathBreadcrumbTrail, SBreadcrumbTrail<FString>)
 					.ButtonContentPadding(FMargin(2.0f, 2.0f))
-					.ButtonStyle(StyleSet.Get()->Get(), "SlateFileDialogs.FlatButton")
-					.DelimiterImage(StyleSet.Get()->GetBrush("SlateFileDialogs.PathDelimiter"))
-					.TextStyle(StyleSet.Get()->Get(), "SlateFileDialogs.PathText")
+					.ButtonStyle(StyleSet->Get(), "SlateFileDialogs.FlatButton")
+					.DelimiterImage(StyleSet->GetBrush("SlateFileDialogs.PathDelimiter"))
+					.TextStyle(StyleSet->Get(), "SlateFileDialogs.PathText")
 					.ShowLeadingDelimiter(false)
 					.InvertTextColorOnHover(false)
 					.OnCrumbClicked(this, &SSlateFileOpenDlg::OnPathClicked)
@@ -424,7 +424,7 @@ void SSlateFileOpenDlg::Construct(const FArguments& InArgs)
 						.ContentPadding(FMargin(0.0f))
 						[
 							SNew(SImage)
-							.Image(StyleSet.Get()->GetBrush("SlateFileDialogs.NewFolder24"))
+							.Image(StyleSet->GetBrush("SlateFileDialogs.NewFolder24"))
 						]
 					]
 
@@ -446,12 +446,12 @@ void SSlateFileOpenDlg::Construct(const FArguments& InArgs)
 							.VAlign(VAlign_Fill)
 							.Padding(FMargin(5.0f,0.0f))
 							.BorderBackgroundColor(FLinearColor(0.1f, 0.1f, 0.1f, 1.0f))
-							.BorderImage(StyleSet.Get()->GetBrush("SlateFileDialogs.WhiteBackground"))
+							.BorderImage(StyleSet->GetBrush("SlateFileDialogs.WhiteBackground"))
 							[
 								SAssignNew(NewDirectoryEditBox, SInlineEditableTextBlock)
-								.Font(StyleSet.Get()->GetFontStyle("SlateFileDialogs.Dialog"))
+								.Font(StyleSet->GetFontStyle("SlateFileDialogs.Dialog"))
 								.IsReadOnly(false)
-								.Text(FText::FromString(""))
+								.Text(FText::GetEmpty())
 								.OnTextCommitted(this, &SSlateFileOpenDlg::OnNewDirectoryCommitted)
 								.OnVerifyTextChanged(this, &SSlateFileOpenDlg::OnNewDirectoryTextChanged)
 							]
@@ -479,7 +479,7 @@ void SSlateFileOpenDlg::Construct(const FArguments& InArgs)
 				.Padding(FMargin(0.0f, 0.0f, 0.0f, 10.0f))
 				[
 					SAssignNew(DirErrorMsg, STextBlock)
-					.Font(StyleSet.Get()->GetFontStyle("SlateFileDialogs.DialogBold"))
+					.Font(StyleSet->GetFontStyle("SlateFileDialogs.DialogBold"))
 					.Justification(ETextJustify::Left)
 					.ColorAndOpacity(FLinearColor::Yellow)
 					.Text(LOCTEXT("SlateFileDialogsDirError", "Unable to create directory!"))
@@ -521,7 +521,7 @@ void SSlateFileOpenDlg::Construct(const FArguments& InArgs)
 								.AutoWidth()
 								[
 									SNew(SImage)
-									.Image(StyleSet.Get()->GetBrush("SlateFileDialogs.Folder24"))
+									.Image(StyleSet->GetBrush("SlateFileDialogs.Folder24"))
 								]
 
 								+ SHorizontalBox::Slot()
@@ -530,8 +530,8 @@ void SSlateFileOpenDlg::Construct(const FArguments& InArgs)
 								.AutoWidth()				
 								[
 									SNew(STextBlock)
-									.Text(FString("Projects"))
-									.Font(StyleSet.Get()->GetFontStyle("SlateFileDialogs.Dialog"))
+									.Text(LOCTEXT("ProjectsLabel", "Projects"))
+									.Font(StyleSet->GetFontStyle("SlateFileDialogs.Dialog"))
 									.Justification(ETextJustify::Left)
 								]
 							]
@@ -557,7 +557,7 @@ void SSlateFileOpenDlg::Construct(const FArguments& InArgs)
 								.AutoWidth()
 								[
 									SNew(SImage)
-									.Image(StyleSet.Get()->GetBrush("SlateFileDialogs.Folder24"))
+									.Image(StyleSet->GetBrush("SlateFileDialogs.Folder24"))
 								]
 
 								+ SHorizontalBox::Slot()
@@ -566,8 +566,8 @@ void SSlateFileOpenDlg::Construct(const FArguments& InArgs)
 								.AutoWidth()				
 								[
 									SNew(STextBlock)
-									.Text(FString("Engine"))
-									.Font(StyleSet.Get()->GetFontStyle("SlateFileDialogs.Dialog"))
+									.Text(LOCTEXT("EngineLabel", "Engine"))
+									.Font(StyleSet->GetFontStyle("SlateFileDialogs.Dialog"))
 									.Justification(ETextJustify::Left)
 								]
 							]
@@ -595,7 +595,7 @@ void SSlateFileOpenDlg::Construct(const FArguments& InArgs)
 						.VAlign(VAlign_Fill)
 						.Padding(FMargin(10.0f))
 						.BorderBackgroundColor(FLinearColor(0.10f, 0.10f, 0.10f, 1.0f))
-						.BorderImage(StyleSet.Get()->GetBrush("SlateFileDialogs.WhiteBackground"))
+						.BorderImage(StyleSet->GetBrush("SlateFileDialogs.WhiteBackground"))
 						[
 							SAssignNew(ListView, SListView<TSharedPtr<FFileEntry>>) // file list scroll
 							.ListItemsSource(&LineItemArray)
@@ -646,8 +646,8 @@ void SSlateFileOpenDlg::Construct(const FArguments& InArgs)
 						.VAlign(VAlign_Fill)
 						[
 							SNew(STextBlock)
-							.Text(FString("Save Filename:"))
-							.Font(StyleSet.Get()->GetFontStyle("SlateFileDialogs.Dialog"))
+							.Text(LOCTEXT("SaveFilenameLabel", "Save Filename:"))
+							.Font(StyleSet->GetFontStyle("SlateFileDialogs.Dialog"))
 							.Justification(ETextJustify::Left)
 						]
 
@@ -668,12 +668,12 @@ void SSlateFileOpenDlg::Construct(const FArguments& InArgs)
 								.VAlign(VAlign_Fill)
 								.Padding(FMargin(5.0f,0.0f))
 								.BorderBackgroundColor(FLinearColor(0.1f, 0.1f, 0.1f, 1.0f))
-								.BorderImage(StyleSet.Get()->GetBrush("SlateFileDialogs.WhiteBackground"))
+								.BorderImage(StyleSet->GetBrush("SlateFileDialogs.WhiteBackground"))
 								[
 									SAssignNew(SaveFilenameEditBox, SInlineEditableTextBlock)
-									.Font(StyleSet.Get()->GetFontStyle("SlateFileDialogs.Dialog"))
+									.Font(StyleSet->GetFontStyle("SlateFileDialogs.Dialog"))
 									.IsReadOnly(false)
-									.Text(FText::FromString(""))
+									.Text(FText::GetEmpty())
 									.OnTextCommitted(this, &SSlateFileOpenDlg::OnFileNameCommitted)
 								]
 							]
@@ -704,8 +704,8 @@ void SSlateFileOpenDlg::Construct(const FArguments& InArgs)
 						.Padding(FMargin(0.0f))
 						[
 							SNew(STextBlock)
-							.Text(FString("Filter:"))
-							.Font(StyleSet.Get()->GetFontStyle("SlateFileDialogs.Dialog"))
+							.Text(LOCTEXT("FilterLabel", "Filter:"))
+							.Font(StyleSet->GetFontStyle("SlateFileDialogs.Dialog"))
 							.Justification(ETextJustify::Left)
 						]
 
@@ -724,7 +724,7 @@ void SSlateFileOpenDlg::Construct(const FArguments& InArgs)
 								.ContentPadding(FMargin(4.0f, 2.0f))
 								//.MaxListHeight(100.0f)
 								.OptionsSource(&FilterNameArray)
-								.Font(StyleSet.Get()->GetFontStyle("SlateFileDialogs.Dialog"))
+								.Font(StyleSet->GetFontStyle("SlateFileDialogs.Dialog"))
 								.OnSelectionChanged(this, &SSlateFileOpenDlg::OnFilterChanged)
 							]
 						]
@@ -773,9 +773,9 @@ void SSlateFileOpenDlg::Construct(const FArguments& InArgs)
 	bDirectoryHasChanged = false;
 	DirectoryWatcher = nullptr;
 
-	if (CurrentPath.Get().Len() > 0 && !CurrentPath.Get().EndsWith("/"))
+	if (CurrentPath.Len() > 0 && !CurrentPath.EndsWith("/"))
 	{
-		CurrentPath = CurrentPath.Get() + TEXT("/");
+		CurrentPath = CurrentPath + TEXT("/");
 	}
 
 #if ENABLE_DIRECTORY_WATCHER	
@@ -794,7 +794,7 @@ END_SLATE_FUNCTION_BUILD_OPTIMIZATION
 void SSlateFileOpenDlg::BuildDirectoryPath()
 {
 	// Clean up path as needed. Fix slashes and convert to absolute path.
-	FString NormPath = CurrentPath.Get();
+	FString NormPath = CurrentPath;
 	FPaths::NormalizeFilename(NormPath);
 	FPaths::RemoveDuplicateSlashes(NormPath);
 	FString AbsPath = FPaths::ConvertRelativePathToFull(NormPath);
@@ -875,7 +875,7 @@ void SSlateFileOpenDlg::RefreshCrumbs()
 }
 
 
-void SSlateFileOpenDlg::OnPathClicked(const FString & NewPath)
+void SSlateFileOpenDlg::OnPathClicked(const FString& NewPath)
 {
 	if (NewPath.Compare("SYSTEM") == 0)
 	{
@@ -1007,43 +1007,43 @@ FReply SSlateFileOpenDlg::OnQuickLinkClick(FSlateFileDlgWindow::EResult ButtonID
 
 void SSlateFileOpenDlg::SetOutputFiles()
 {
-	if (OutNames.Get() != nullptr)
+	if (OutNames != nullptr)
 	{
-		OutNames.Get()->Empty();
+		OutNames->Empty();
 
-		if (bSaveFile.Get())
+		if (bSaveFile)
 		{
-			FString Path = CurrentPath.Get() + SaveFilename;
-			OutNames.Get()->Add(Path);
+			FString Path = CurrentPath + SaveFilename;
+			OutNames->Add(Path);
 		}
 		else
 		{
 			TArray<TSharedPtr<FFileEntry>> SelectedItems = ListView->GetSelectedItems();
 
-			if (bDirectoriesOnly.Get())
+			if (bDirectoriesOnly)
 			{
 				if (SelectedItems.Num() > 0)
 				{
-					FString Path = CurrentPath.Get() + SelectedItems[0]->Label;
-					OutNames.Get()->Add(Path);
+					FString Path = CurrentPath + SelectedItems[0]->Label;
+					OutNames->Add(Path);
 				}
 				else
 				{
 					// select the current directory
-					OutNames.Get()->Add(CurrentPath.Get());
+					OutNames->Add(CurrentPath);
 				}
 			}
 			else
 			{
 				for (int32 i = 0; i < SelectedItems.Num(); i++)
 				{
-						FString Path = CurrentPath.Get() + SelectedItems[i]->Label;
-						OutNames.Get()->Add(Path);
+					FString Path = CurrentPath + SelectedItems[i]->Label;
+					OutNames->Add(Path);
 				}
 
-				if (OutFilterIndex.Get() != nullptr)
+				if (OutFilterIndex != nullptr)
 				{
-					*(OutFilterIndex.Get()) = FilterIndex;
+					*(OutFilterIndex) = FilterIndex;
 				}
 			}
 		}
@@ -1060,11 +1060,11 @@ FReply SSlateFileOpenDlg::OnAcceptCancelClick(FSlateFileDlgWindow::EResult Butto
 	else
 	{
 		FString Path = "";
-		OutNames.Get()->Add(Path);
+		OutNames->Add(Path);
 	}
 
 	UserResponse = ButtonID;
-	ParentWindow.Get().Pin()->RequestDestroyWindow();
+	ParentWindow.Pin()->RequestDestroyWindow();
 
 	return FReply::Handled();
 }
@@ -1072,7 +1072,7 @@ FReply SSlateFileOpenDlg::OnAcceptCancelClick(FSlateFileDlgWindow::EResult Butto
 
 FReply SSlateFileOpenDlg::OnDirSublevelClick(int32 Level)
 {
-	DirectoryNodesArray[DirNodeIndex].TextBlock->SetFont(StyleSet.Get()->GetFontStyle("SlateFileDialogs.Dialog"));
+	DirectoryNodesArray[DirNodeIndex].TextBlock->SetFont(StyleSet->GetFontStyle("SlateFileDialogs.Dialog"));
 
 	FString NewPath = TEXT("/");
 
@@ -1086,7 +1086,7 @@ FReply SSlateFileOpenDlg::OnDirSublevelClick(int32 Level)
 	bNeedsBuilding = true;
 
 	DirNodeIndex = Level;
-	DirectoryNodesArray[DirNodeIndex].TextBlock->SetFont(StyleSet.Get()->GetFontStyle("SlateFileDialogs.DialogBold"));
+	DirectoryNodesArray[DirNodeIndex].TextBlock->SetFont(StyleSet->GetFontStyle("SlateFileDialogs.DialogBold"));
 
 	return FReply::Handled();
 }
@@ -1151,18 +1151,18 @@ void SSlateFileOpenDlg::ReadDir(bool bIsRefresh)
 
 	FSlateFileDialogVisitor DirVisitor(FilesArray, FoldersArray, FilterList);
 
-	FileManager.IterateDirectory(*CurrentPath.Get(), DirVisitor);
+	FileManager.IterateDirectory(*CurrentPath, DirVisitor);
 	
 	FilesArray.Sort(FFileEntry::ConstPredicate);
 	FoldersArray.Sort(FFileEntry::ConstPredicate);
 
 	if (DirectoryWatcher && !bIsRefresh)
 	{
-		DirectoryWatcher->RegisterDirectoryChangedCallback_Handle(CurrentPath.Get(),
+		DirectoryWatcher->RegisterDirectoryChangedCallback_Handle(CurrentPath,
 						IDirectoryWatcher::FDirectoryChanged::CreateRaw(this, &SSlateFileOpenDlg::OnDirectoryChanged),
 						OnDialogDirectoryChangedDelegateHandle, IDirectoryWatcher::WatchOptions::IncludeDirectoryChanges | IDirectoryWatcher::WatchOptions::IgnoreChangesInSubtree);
 
-		RegisteredPath = CurrentPath.Get();
+		RegisteredPath = CurrentPath;
 	}
 }
 
@@ -1185,7 +1185,7 @@ void SSlateFileOpenDlg::RebuildFileTable()
 	}
 
 	// file entries
-	if (bDirectoriesOnly.Get() == false)
+	if (bDirectoriesOnly == false)
 	{
 		for (int32 i = 0; i < FilesArray.Num(); i++)
 		{
@@ -1200,7 +1200,7 @@ TSharedRef<ITableRow> SSlateFileOpenDlg::OnGenerateWidgetForList(TSharedPtr<FFil
 {
 	return SNew(SSlateFileDialogRow, OwnerTable)
 			.DialogItem(Item)
-			.StyleSet(StyleSet.Get());
+			.StyleSet(StyleSet);
 }
 
 
@@ -1208,7 +1208,7 @@ void SSlateFileOpenDlg::OnItemDoubleClicked(TSharedPtr<FFileEntry> Item)
 {
 	if (Item->bIsDirectory)
 	{
-		CurrentPath = CurrentPath.Get() + Item->Label + TEXT("/");
+		CurrentPath = CurrentPath + Item->Label + TEXT("/");
 		bNeedsBuilding = true;
 		bRebuildDirPath = true;
 	}
@@ -1216,7 +1216,7 @@ void SSlateFileOpenDlg::OnItemDoubleClicked(TSharedPtr<FFileEntry> Item)
 	{
 		SetOutputFiles();
 		UserResponse = FSlateFileDlgWindow::Accept;
-		ParentWindow.Get().Pin()->RequestDestroyWindow();
+		ParentWindow.Pin()->RequestDestroyWindow();
 	}
 }
 
@@ -1243,7 +1243,7 @@ void SSlateFileOpenDlg::SetDefaultFile(FString DefaultFile)
 }
 
 
-void SSlateFileOpenDlg::OnFileNameCommitted(const FText & InText, ETextCommit::Type InCommitType)
+void SSlateFileOpenDlg::OnFileNameCommitted(const FText& InText, ETextCommit::Type InCommitType)
 {
 	// update edit box unless user choose to escape out
 	if (InCommitType != ETextCommit::OnCleared)
@@ -1276,7 +1276,7 @@ void SSlateFileOpenDlg::OnItemSelected(TSharedPtr<FFileEntry> Item, ESelectInfo:
 {
 	if (Item.IsValid())
 	{
-		if (!bDirectoriesOnly.Get())
+		if (!bDirectoriesOnly)
 		{
 			TArray<TSharedPtr<FFileEntry>> SelectedItems = ListView->GetSelectedItems();
 
@@ -1289,7 +1289,7 @@ void SSlateFileOpenDlg::OnItemSelected(TSharedPtr<FFileEntry> Item, ESelectInfo:
 			}
 		}
 
-		if (bSaveFile.Get() && !Item->bIsDirectory)
+		if (bSaveFile && !Item->bIsDirectory)
 		{
 			SetDefaultFile(Item->Label);
 		}
@@ -1301,12 +1301,12 @@ void SSlateFileOpenDlg::ParseFilters()
 {
 	if (FilterCombo.IsValid() && FilterHBox.IsValid())
 	{
-		if (Filters.Get().Len() > 0)
+		if (Filters.Len() > 0)
 		{
 			if (FilterNameArray.Num() == 0)
 			{
 				TCHAR Temp[MAX_FILTER_LENGTH];
-				FCString::Strcpy(Temp, Filters.Get().Len(), *Filters.Get());
+				FCString::Strcpy(Temp, Filters.Len(), *Filters);
 
 				// break path into tokens
 				TCHAR *ContextStr = nullptr;
@@ -1341,7 +1341,7 @@ void SSlateFileOpenDlg::ParseFilters()
 bool SSlateFileOpenDlg::GetFilterExtension(FString &OutString)
 {
 	// check to see if filters were given
-	if (Filters.Get().Len() == 0)
+	if (Filters.Len() == 0)
 	{
 		OutString = "";
 		return false;
@@ -1417,7 +1417,7 @@ FReply SSlateFileOpenDlg::OnNewDirectoryAcceptCancelClick(FSlateFileDlgWindow::E
 		if (NewDirectoryName.Len() > 0)
 		{
 			IPlatformFile &PlatformFile = FPlatformFileManager::Get().GetPlatformFile();
-			FString DirPath = CurrentPath.Get() + NewDirectoryName;
+			FString DirPath = CurrentPath + NewDirectoryName;
 
 			if (!PlatformFile.CreateDirectory(*DirPath))
 			{
diff --git a/Engine/Source/Developer/SlateFileDialogs/Private/SlateFileDlgWindow.h b/Engine/Source/Developer/SlateFileDialogs/Private/SlateFileDlgWindow.h
--- a/Engine/Source/Developer/SlateFileDialogs/Private/SlateFileDlgWindow.h
+++ b/Engine/Source/Developer/SlateFileDialogs/Private/SlateFileDlgWindow.h
@@ -24,12 +24,12 @@ struct FFileEntry
 	FFileEntry() { }
 
 	FFileEntry(FString InLabel, FString InModDate, FString InFileSize, bool InIsDirectory)
+		: Label(MoveTemp(InLabel))
+		, ModDate(MoveTemp(InModDate))
+		, FileSize(MoveTemp(InFileSize))
+		, bIsSelected(false)
+		, bIsDirectory(InIsDirectory)
 	{
-		Label = InLabel;
-		ModDate = FString(InModDate);
-		FileSize = FString(InFileSize);
-		bIsSelected = false;
-		bIsDirectory = InIsDirectory;
 	}
 
 	inline static bool ConstPredicate ( const TSharedPtr<FFileEntry> Entry1, const TSharedPtr<FFileEntry> Entry2)
@@ -109,27 +109,27 @@ public:
 		_StyleSet(nullptr)
 	{}
 
-	SLATE_ATTRIBUTE(FString, CurrentPath)
-	SLATE_ATTRIBUTE(FString, Filters)
-	SLATE_ATTRIBUTE(bool, bMultiSelectEnabled)
-	SLATE_ATTRIBUTE(FString, WindowTitleText)
-	SLATE_ATTRIBUTE(FText, AcceptText)
-	SLATE_ATTRIBUTE(bool, bDirectoriesOnly)
-	SLATE_ATTRIBUTE(bool, bSaveFile)
-	SLATE_ATTRIBUTE(TArray<FString> *, OutNames)
-	SLATE_ATTRIBUTE(int32 *, OutFilterIndex)
-	SLATE_ATTRIBUTE(TWeakPtr<SWindow>, ParentWindow)
-	SLATE_ATTRIBUTE(FSlateFileDialogsStyle *, StyleSet)
+	SLATE_ARGUMENT(FString, CurrentPath)
+	SLATE_ARGUMENT(FString, Filters)
+	SLATE_ARGUMENT(bool, bMultiSelectEnabled)
+	SLATE_ARGUMENT(FString, WindowTitleText)
+	SLATE_ARGUMENT(FText, AcceptText)
+	SLATE_ARGUMENT(bool, bDirectoriesOnly)
+	SLATE_ARGUMENT(bool, bSaveFile)
+	SLATE_ARGUMENT(TArray<FString>*, OutNames)
+	SLATE_ARGUMENT(int32*, OutFilterIndex)
+	SLATE_ARGUMENT(TWeakPtr<SWindow>, ParentWindow)
+	SLATE_ARGUMENT(FSlateFileDialogsStyle*, StyleSet)
 	SLATE_END_ARGS()
 
 	void Construct(const FArguments& InArgs);
 	FSlateFileDlgWindow::EResult GetResponse() { return UserResponse; }
-	void SetOutNames(TArray<FString> *Ptr) { OutNames = Ptr; }
-	void SetOutFilterIndex(int32 *InOutFilterIndex) { OutFilterIndex = InOutFilterIndex; }
+	void SetOutNames(TArray<FString>* Ptr) { OutNames = Ptr; }
+	void SetOutFilterIndex(int32* InOutFilterIndex) { OutFilterIndex = InOutFilterIndex; }
 	void SetDefaultFile(FString DefaultFile);
 
 private:	
-	void OnPathClicked( const FString & CrumbData );
+	void OnPathClicked( const FString& CrumbData );
 	TSharedPtr<SWidget> OnGetCrumbDelimiterContent(const FString& CrumbData) const;
 	void RebuildFileTable();
 	void BuildDirectoryPath();
@@ -148,9 +148,9 @@ private:
 	FReply OnQuickLinkClick(FSlateFileDlgWindow::EResult ButtonID);
 	FReply OnDirSublevelClick(int32 Level);
 	void OnDirectoryChanged(const TArray <FFileChangeData> &FileChanges);	
-	void OnFileNameCommitted(const FText & InText, ETextCommit::Type InCommitType);
+	void OnFileNameCommitted(const FText& InText, ETextCommit::Type InCommitType);
 
-	void OnNewDirectoryCommitted(const FText & InText, ETextCommit::Type InCommitType);
+	void OnNewDirectoryCommitted(const FText& InText, ETextCommit::Type InCommitType);
 	FReply OnNewDirectoryClick();
 	bool OnNewDirectoryTextChanged(const FText &InText, FText &ErrorMsg);
 	FReply OnNewDirectoryAcceptCancelClick(FSlateFileDlgWindow::EResult ButtonID);
@@ -193,17 +193,17 @@ private:
 	int32 DirNodeIndex;
 	FString SaveFilename;
 
-	TAttribute<TWeakPtr<SWindow>> ParentWindow;
-	TAttribute<FString> CurrentPath;
-	TAttribute<FString> Filters;
-	TAttribute<FString> WindowTitleText;
-	TAttribute<bool> bMultiSelectEnabled;
-	TAttribute<TArray<FString> *> OutNames;
-	TAttribute<int32 *> OutFilterIndex;
-	TAttribute<bool> bDirectoriesOnly;
-	TAttribute<bool> bSaveFile;
-	TAttribute<FText> AcceptText;
-	TAttribute<FSlateFileDialogsStyle *> StyleSet;
+	TWeakPtr<SWindow> ParentWindow;
+	FString CurrentPath;
+	FString Filters;
+	FString WindowTitleText;
+	bool bMultiSelectEnabled;
+	TArray<FString>* OutNames;
+	int32* OutFilterIndex;
+	bool bDirectoriesOnly;
+	bool bSaveFile;
+	FText AcceptText;
+	FSlateFileDialogsStyle* StyleSet;
 
 	IDirectoryWatcher *DirectoryWatcher;
 	FDelegateHandle OnDialogDirectoryChangedDelegateHandle;
@@ -220,7 +220,7 @@ public:
 	
 	SLATE_BEGIN_ARGS(SSlateFileDialogRow) { }
 	SLATE_ARGUMENT(SSlateFileDialogItemPtr, DialogItem)
-	SLATE_ATTRIBUTE(FSlateFileDialogsStyle *, StyleSet)
+	SLATE_ARGUMENT(FSlateFileDialogsStyle *, StyleSet)
 	SLATE_END_ARGS()
 
 
@@ -237,7 +237,7 @@ public:
 	BEGIN_SLATE_FUNCTION_BUILD_OPTIMIZATION
 	virtual TSharedRef<SWidget> GenerateWidgetForColumn( const FName& ColumnName ) override
 	{
-		FSlateFontInfo ItemFont = StyleSet.Get()->GetFontStyle("SlateFileDialogs.Dialog");
+		FSlateFontInfo ItemFont = StyleSet->GetFontStyle("SlateFileDialogs.Dialog");
 		struct EVisibility FolderIconVisibility = DialogItem->bIsDirectory ? EVisibility::Visible : EVisibility::Hidden;
 
 		if (ColumnName == TEXT("Pathname"))
@@ -251,7 +251,7 @@ public:
 					.Padding(FMargin(5.0f, 2.0f))
 					[
 						SNew(SImage)
-						.Image(StyleSet.Get()->GetBrush("SlateFileDialogs.Folder16"))
+						.Image(StyleSet->GetBrush("SlateFileDialogs.Folder16"))
 						.Visibility(FolderIconVisibility)
 					]
 
@@ -262,20 +262,20 @@ public:
 					.Padding(FMargin(5.0f, 0.0f, 0.0f, 0.0f))
 					[
 						SNew(STextBlock)
-						.Text(DialogItem->Label)
+						.Text(FText::FromString(DialogItem->Label))
 						.Font(ItemFont)
 					];
 		}
 		else if (ColumnName == TEXT("ModDate"))
 		{
 			return SNew(STextBlock)
-					.Text(DialogItem->ModDate)
+					.Text(FText::FromString(DialogItem->ModDate))
 					.Font(ItemFont);
 		}
 		else if (ColumnName == TEXT("FileSize"))
 		{
 			return SNew(STextBlock)
-					.Text(DialogItem->FileSize)
+					.Text(FText::FromString(DialogItem->FileSize))
 					.Font(ItemFont);
 		}
 		else
@@ -288,7 +288,7 @@ public:
 
 private:
 	SSlateFileDialogItemPtr DialogItem;
-	TAttribute<FSlateFileDialogsStyle *> StyleSet;
+	FSlateFileDialogsStyle* StyleSet;
 };
 
 #undef LOCTEXT_NAMESPACE
