SlateDialogs: Fix extra slash being appended (UE-18138).

- Somewhat unexpectedly, this slash also broke packaging (UE-18819).

diff --git a/Engine/Source/Developer/SlateFileDialogs/Private/SlateFileDlgWindow.cpp b/Engine/Source/Developer/SlateFileDialogs/Private/SlateFileDlgWindow.cpp
--- a/Engine/Source/Developer/SlateFileDialogs/Private/SlateFileDlgWindow.cpp
+++ b/Engine/Source/Developer/SlateFileDialogs/Private/SlateFileDlgWindow.cpp
@@ -293,7 +293,7 @@ bool FSlateFileDlgWindow::OpenDirectoryDialog(const void* ParentWindowHandle, co
 	FSlateApplication::Get().AddModalWindow(ModalWindow, NULL);
 	bool RC = (DialogWidget->GetResponse() == EResult::Accept && TempOut.Num() > 0);
 
-	if (TempOut.Num() > 0)
+	if (TempOut.Num() > 0 && !TempOut[0].EndsWith(TEXT("/")))
 	{
 		OutFoldername = TempOut[0] + TEXT("/");
 	}
