Linux: fix an edge case packaging failure (UE-19933).

- Seems that sometimes packaging fails because destination directory can be passed as -archivedirectory=../../../Foo (may happen if both Foo and RootDir() are in the same directory) and UAT incorrectly tries to create /Foo instead of Foo.
- The easier fix (in terms of testing) is to make sure that the directory paths are absolute.

diff --git a/Engine/Source/Developer/SlateFileDialogs/Private/SlateFileDlgWindow.cpp b/Engine/Source/Developer/SlateFileDialogs/Private/SlateFileDlgWindow.cpp
--- a/Engine/Source/Developer/SlateFileDialogs/Private/SlateFileDlgWindow.cpp
+++ b/Engine/Source/Developer/SlateFileDialogs/Private/SlateFileDlgWindow.cpp
@@ -295,7 +295,7 @@ bool FSlateFileDlgWindow::OpenDirectoryDialog(const void* ParentWindowHandle, co
 
 	if (TempOut.Num() > 0)
 	{
-		OutFoldername = TempOut[0];
+		OutFoldername = FPaths::ConvertRelativePathToFull(TempOut[0]);
 		if (!OutFoldername.EndsWith(TEXT("/")))
 		{
 			OutFoldername += TEXT("/");
