Removed unneeded dependency from SlateFileDialogs

diff --git a/Engine/Source/Developer/SlateFileDialogs/SlateFileDialogs.Build.cs b/Engine/Source/Developer/SlateFileDialogs/SlateFileDialogs.Build.cs
--- a/Engine/Source/Developer/SlateFileDialogs/SlateFileDialogs.Build.cs
+++ b/Engine/Source/Developer/SlateFileDialogs/SlateFileDialogs.Build.cs
@@ -10,7 +10,6 @@ public class SlateFileDialogs : ModuleRules
             new string[] {
                 "Core",
                 "InputCore",
-                "AppFramework",
                 "Slate",
                 "SlateCore",
                 "DirectoryWatcher",
