[NO NON-PII COMMIT MESSAGE]

diff --git a/Engine/Source/Developer/SlateFileDialogs/Private/SlateFileDialogsStyles.cpp b/Engine/Source/Developer/SlateFileDialogs/Private/SlateFileDialogsStyles.cpp
--- a/Engine/Source/Developer/SlateFileDialogs/Private/SlateFileDialogsStyles.cpp
+++ b/Engine/Source/Developer/SlateFileDialogs/Private/SlateFileDialogsStyles.cpp
@@ -81,9 +81,21 @@ TSharedPtr< FSlateStyleSet > FSlateFileDialogsStyle::Create()
 
 	Style->Set("SlateFileDialogs.Folder16", new IMAGE_BRUSH("SlateFileDialogs/Icons/icon_file_folder_16x", Icon16x16));
 	Style->Set("SlateFileDialogs.Folder24", new IMAGE_BRUSH("SlateFileDialogs/Icons/icon_file_folder_40x", Icon24x24));
-	Style->Set("SlateFileDialogs.NewFolder24", new IMAGE_BRUSH("SlateFileDialogs/Icons/icon_new_folder_40x", Icon24x24));		
+	Style->Set("SlateFileDialogs.NewFolder24", new IMAGE_BRUSH("SlateFileDialogs/Icons/icon_new_folder_40x", Icon24x24));
+	Style->Set("SlateFileDialogs.BrowseBack24", new IMAGE_BRUSH("SlateFileDialogs/Common/back_arrow_40x", Icon24x24));
+	Style->Set("SlateFileDialogs.BrowseForward24", new IMAGE_BRUSH("SlateFileDialogs/Common/forward_arrow_40x", Icon24x24));
 	Style->Set("SlateFileDialogs.WhiteBackground", new IMAGE_BRUSH("SlateFileDialogs/Common/Window/WindowWhite", Icon64x64));
 
+	Style->Set("SlateFileDialogs.UAsset16", new IMAGE_BRUSH("SlateFileDialogs/Icons/icon_uasset_24x", Icon16x16));
+	Style->Set("SlateFileDialogs.UProject16", new IMAGE_BRUSH("SlateFileDialogs/Icons/icon_uproject_24x", Icon16x16));
+	Style->Set("SlateFileDialogs.Model3D", new IMAGE_BRUSH("SlateFileDialogs/Icons/icon_model_24x", Icon16x16));
+	Style->Set("SlateFileDialogs.Video", new IMAGE_BRUSH("SlateFileDialogs/Icons/icon_video_24x", Icon16x16));
+	Style->Set("SlateFileDialogs.Audio", new IMAGE_BRUSH("SlateFileDialogs/Icons/icon_audio_24x", Icon16x16));
+	Style->Set("SlateFileDialogs.Image", new IMAGE_BRUSH("SlateFileDialogs/Icons/icon_image_24x", Icon16x16));
+	Style->Set("SlateFileDialogs.TextFile", new IMAGE_BRUSH("SlateFileDialogs/Icons/icon_text_24x", Icon16x16));
+	Style->Set("SlateFileDialogs.PlaceHolder", new IMAGE_BRUSH("SlateFileDialogs/Icons/icon_skull_16x", Icon16x16));
+
+
 	Style->Set( "SlateFileDialogs.PathDelimiter", new IMAGE_BRUSH( "SlateFileDialogs/Common/SmallArrowRight", Icon10x10 ) );
 
 	Style->Set( "SlateFileDialogs.PathText", FTextBlockStyle()
diff --git a/Engine/Source/Developer/SlateFileDialogs/Private/SlateFileDlgWindow.cpp b/Engine/Source/Developer/SlateFileDialogs/Private/SlateFileDlgWindow.cpp
--- a/Engine/Source/Developer/SlateFileDialogs/Private/SlateFileDlgWindow.cpp
+++ b/Engine/Source/Developer/SlateFileDialogs/Private/SlateFileDlgWindow.cpp
@@ -365,7 +365,7 @@ void SSlateFileOpenDlg::Construct(const FArguments& InArgs)
 	FilterIndex = 0;
 
 	ESelectionMode::Type SelectMode = bMultiSelectEnabled ? ESelectionMode::Multi : ESelectionMode::Single;
-	struct EVisibility SaveFilenameVisibility = bSaveFile ? EVisibility::Visible : EVisibility::Collapsed;
+	struct EVisibility SaveFilenameVisibility = bDirectoriesOnly ? EVisibility::Collapsed : EVisibility::Visible;
 
 	this->ChildSlot
 		[
@@ -415,6 +415,40 @@ void SSlateFileOpenDlg::Construct(const FArguments& InArgs)
 				[
 					SNew(SHorizontalBox)
 
+					+ SHorizontalBox::Slot()
+					.Padding(FMargin(0.0f, 0.0f, 10.0f, 0.0f))
+					.AutoWidth()
+					.HAlign(HAlign_Fill)
+					.VAlign(VAlign_Fill)
+					[
+						SNew(SButton)
+						.HAlign(HAlign_Center)
+						.VAlign(VAlign_Center)
+						.OnClicked(this, &SSlateFileOpenDlg::OnGoBackClick)
+						.ContentPadding(FMargin(0.0f))
+						[
+							SNew(SImage)
+							.Image(StyleSet->GetBrush("SlateFileDialogs.BrowseBack24"))
+						]
+					]
+
+					+ SHorizontalBox::Slot()
+					.Padding(FMargin(0.0f, 0.0f, 40.0f, 0.0f))
+					.AutoWidth()
+					.HAlign(HAlign_Fill)
+					.VAlign(VAlign_Fill)
+					[
+						SNew(SButton)
+						.HAlign(HAlign_Center)
+						.VAlign(VAlign_Center)
+						.OnClicked(this, &SSlateFileOpenDlg::OnGoForwardClick)
+						.ContentPadding(FMargin(0.0f))
+						[
+							SNew(SImage)
+							.Image(StyleSet->GetBrush("SlateFileDialogs.BrowseForward24"))
+						]
+					]
+
 					+ SHorizontalBox::Slot()
 					.Padding(FMargin(0.0f))
 					.AutoWidth()
@@ -650,7 +684,7 @@ void SSlateFileOpenDlg::Construct(const FArguments& InArgs)
 						.VAlign(VAlign_Fill)
 						[
 							SNew(STextBlock)
-							.Text(LOCTEXT("SaveFilenameLabel", "Save Filename:"))
+							.Text(LOCTEXT("FilenameLabel", "Filename:"))
 							.Font(StyleSet->GetFontStyle("SlateFileDialogs.Dialog"))
 							.Justification(ETextJustify::Left)
 						]
@@ -782,6 +816,9 @@ void SSlateFileOpenDlg::Construct(const FArguments& InArgs)
 		CurrentPath = CurrentPath + TEXT("/");
 	}
 
+	HistoryIndex = 0;
+	History.Add(CurrentPath);
+
 #if ENABLE_DIRECTORY_WATCHER	
 	if (!FModuleManager::Get().IsModuleLoaded("DirectoryWatcher"))
 	{
@@ -892,6 +929,14 @@ void SSlateFileOpenDlg::OnPathClicked(const FString& NewPath)
 	bRebuildDirPath = true;
 	bNeedsBuilding = true;
 
+	if ((History.Num()-HistoryIndex-1) > 0)
+	{
+		History.RemoveAt(HistoryIndex+1, History.Num()-HistoryIndex-1, true);
+	}
+
+	History.Add(CurrentPath);
+	HistoryIndex++;
+
 	RefreshCrumbs();
 }
 
@@ -901,6 +946,14 @@ void SSlateFileOpenDlg::OnPathMenuItemClicked( FString ClickedPath )
 	CurrentPath = ClickedPath;
 	bRebuildDirPath = true;
 	bNeedsBuilding = true;
+
+	if ((History.Num()-HistoryIndex-1) > 0)
+	{
+		History.RemoveAt(HistoryIndex+1, History.Num()-HistoryIndex-1, true);
+	}
+
+	History.Add(CurrentPath);
+	HistoryIndex++;
 	
 	RefreshCrumbs();
 }
@@ -1002,6 +1055,14 @@ FReply SSlateFileOpenDlg::OnQuickLinkClick(FSlateFileDlgWindow::EResult ButtonID
 	{
 		CurrentPath = FPaths::EngineDir();
 	}
+	
+	if ((History.Num()-HistoryIndex-1) > 0)
+	{
+		History.RemoveAt(HistoryIndex+1, History.Num()-HistoryIndex-1, true);
+	}
+
+	History.Add(CurrentPath);
+	HistoryIndex++;
 
 	bNeedsBuilding = true;
 	bRebuildDirPath = true;
@@ -1013,42 +1074,35 @@ void SSlateFileOpenDlg::SetOutputFiles()
 {
 	if (OutNames != nullptr)
 	{
+		TArray<FString> NamesArray;
+		ParseTextField(NamesArray, SaveFilename);
+
 		OutNames->Empty();
 
-		if (bSaveFile)
+		if (bDirectoriesOnly)
 		{
-			FString Path = CurrentPath + SaveFilename;
-			OutNames->Add(Path);
+			if (NamesArray.Num() > 0)
+			{
+				FString Path = CurrentPath + NamesArray[0];
+				OutNames->Add(Path);
+			}
+			else
+			{
+				// select the current directory
+				OutNames->Add(CurrentPath);
+			}
 		}
 		else
 		{
-			TArray<TSharedPtr<FFileEntry>> SelectedItems = ListView->GetSelectedItems();
-
-			if (bDirectoriesOnly)
+			for (int32 i=0; i < NamesArray.Num(); i++)
 			{
-				if (SelectedItems.Num() > 0)
-				{
-					FString Path = CurrentPath + SelectedItems[0]->Label;
-					OutNames->Add(Path);
-				}
-				else
-				{
-					// select the current directory
-					OutNames->Add(CurrentPath);
-				}
+				FString Path = CurrentPath + NamesArray[i];
+				OutNames->Add(Path);
 			}
-			else
-			{
-				for (int32 i = 0; i < SelectedItems.Num(); i++)
-				{
-					FString Path = CurrentPath + SelectedItems[i]->Label;
-					OutNames->Add(Path);
-				}
 
-				if (OutFilterIndex != nullptr)
-				{
-					*(OutFilterIndex) = FilterIndex;
-				}
+			if (OutFilterIndex != nullptr)
+			{
+				*(OutFilterIndex) = FilterIndex;
 			}
 		}
 	}
@@ -1063,8 +1117,10 @@ FReply SSlateFileOpenDlg::OnAcceptCancelClick(FSlateFileDlgWindow::EResult Butto
 	}
 	else
 	{
-		FString Path = "";
-		OutNames->Add(Path);
+		if (OutNames != nullptr)
+		{
+			OutNames->Empty();
+		}
 	}
 
 	UserResponse = ButtonID;
@@ -1212,9 +1268,19 @@ void SSlateFileOpenDlg::OnItemDoubleClicked(TSharedPtr<FFileEntry> Item)
 {
 	if (Item->bIsDirectory)
 	{
+		SetDefaultFile(FString(""));
+
 		CurrentPath = CurrentPath + Item->Label + TEXT("/");
 		bNeedsBuilding = true;
 		bRebuildDirPath = true;
+
+		if ((History.Num()-HistoryIndex-1) > 0)
+		{
+			History.RemoveAt(HistoryIndex+1, History.Num()-HistoryIndex-1, true);
+		}
+
+		History.Add(CurrentPath);
+		HistoryIndex++;
 	}
 	else
 	{
@@ -1239,10 +1305,91 @@ void SSlateFileOpenDlg::OnFilterChanged(TSharedPtr<FString> NewValue, ESelectInf
 	bNeedsBuilding = true;
 }
 
+void SSlateFileOpenDlg::ParseTextField(TArray<FString> &FilenameArray, FString Files)
+{
+	FString FileList = Files;
+	FileList.Trim();
+	FileList.TrimTrailing();
+
+	FilenameArray.Empty();
+
+	if  (FileList.Len() > 0 && FileList[0] == TCHAR('"'))
+	{
+		FString TempName;
+		SaveFilename.Empty();
+
+		for (int32 i = 0; i < FileList.Len(); )
+		{
+			// find opening quote (")
+			for (; i < FileList.Len() && FileList[i] != TCHAR('"'); i++);
+			
+			if (i >= FileList.Len())
+			{
+				break;
+			}
+			
+			// copy name until closing quote is found.
+			TempName.Empty();
+			for (i++; i < FileList.Len() && FileList[i] != TCHAR('"'); i++)
+			{
+				TempName.AppendChar(FileList[i]);
+			}
+
+			if (i >= FileList.Len())
+			{
+				break;
+			}
+			
+			// check to see if file exists or if we're trying to save a file. if so, add it to list.
+			if (FPaths::FileExists(CurrentPath + TempName) || bSaveFile)
+			{
+				FilenameArray.Add(TempName);
+			}
+
+			// if multiselect is off, don't bother parsing out any additional file names.
+			if (!bMultiSelectEnabled)
+			{
+				break;
+			}
+
+			i++;
+		}
+	}
+	else
+	{
+		FilenameArray.Add(Files);
+	}
+}
 
 void SSlateFileOpenDlg::SetDefaultFile(FString DefaultFile)
 {
-	SaveFilename = DefaultFile;
+	FString FileList = DefaultFile;
+	FileList.Trim();
+	FileList.TrimTrailing();
+
+	if  (FileList.Len() > 0 && FileList[0] == TCHAR('"'))
+	{
+		TArray<FString> NamesArray;
+		ParseTextField(NamesArray, FileList);
+
+		SaveFilename.Empty();
+
+		for (int32 i = 0; i < NamesArray.Num(); i++)
+		{	
+			SaveFilename = SaveFilename + TEXT("\"") + NamesArray[i] + TEXT("\" ");
+
+			// if multiselect is off, don't bother adding any additional file names.
+			if (!bMultiSelectEnabled)
+			{
+				break;
+			}
+		}
+	}
+	else
+	{
+		SaveFilename = FileList;
+	}
+
 	SaveFilenameEditBox->SetText(SaveFilename);
 }
 
@@ -1253,11 +1400,10 @@ void SSlateFileOpenDlg::OnFileNameCommitted(const FText& InText, ETextCommit::Ty
 	if (InCommitType != ETextCommit::OnCleared)
 	{
 		FString Extension;
-
 		SaveFilename = InText.ToString();
 
 		// get current filter extension
-		if (GetFilterExtension(Extension))
+		if (!bDirectoriesOnly && GetFilterExtension(Extension))
 		{
 			// append extension to filename if user left it off
 			if (!SaveFilename.EndsWith(Extension, ESearchCase::CaseSensitive) && Extension.Compare(".*") != 0)
@@ -1266,12 +1412,9 @@ void SSlateFileOpenDlg::OnFileNameCommitted(const FText& InText, ETextCommit::Ty
 			}
 		}
 
-		// remove any whitespace from start and/or end of filename
-		SaveFilename.Trim();
-		SaveFilename.TrimTrailing();
+		ListView->ClearSelection();
 
-		// update edit box content
-		SaveFilenameEditBox->SetText(SaveFilename);
+		SetDefaultFile(SaveFilename);
 	}
 }
 
@@ -1280,22 +1423,32 @@ void SSlateFileOpenDlg::OnItemSelected(TSharedPtr<FFileEntry> Item, ESelectInfo:
 {
 	if (Item.IsValid())
 	{
+		FString FileList;
+
 		if (!bDirectoriesOnly)
 		{
 			TArray<TSharedPtr<FFileEntry>> SelectedItems = ListView->GetSelectedItems();
-
+			
 			for (int32 i = 0; i < SelectedItems.Num(); i++)
 			{
 				if (SelectedItems[i]->bIsDirectory)
 				{
 					ListView->SetItemSelection(SelectedItems[i], false, ESelectInfo::Direct);
 				}
+				else
+				{
+					FileList = FileList + TEXT("\"") + SelectedItems[i]->Label + TEXT("\" ");
+				}
 			}
 		}
-
-		if (bSaveFile && !Item->bIsDirectory)
+		else
 		{
-			SetDefaultFile(Item->Label);
+			FileList = Item->Label;
+		}
+	
+		if (bDirectoriesOnly == Item->bIsDirectory)
+		{
+			SetDefaultFile(FileList);
 		}
 	}
 }
@@ -1442,4 +1595,38 @@ FReply SSlateFileOpenDlg::OnNewDirectoryAcceptCancelClick(FSlateFileDlgWindow::E
 	return FReply::Handled();
 }
 
+
+FReply SSlateFileOpenDlg::OnGoForwardClick()
+{
+	if ((HistoryIndex+1) < History.Num())
+	{
+		SetDefaultFile(FString(""));
+
+		HistoryIndex++;
+		CurrentPath = History[HistoryIndex];
+		bNeedsBuilding = true;
+		bRebuildDirPath = true;
+		bDirectoryHasChanged = false;
+	}
+
+	return FReply::Handled();
+}
+
+
+FReply SSlateFileOpenDlg::OnGoBackClick()
+{
+	if (HistoryIndex > 0)
+	{
+		SetDefaultFile(FString(""));
+
+		HistoryIndex--;
+		CurrentPath = History[HistoryIndex];
+		bNeedsBuilding = true;
+		bRebuildDirPath = true;
+		bDirectoryHasChanged = false;	
+	}
+
+	return FReply::Handled();
+}
+
 #undef LOCTEXT_NAMESPACE
diff --git a/Engine/Source/Developer/SlateFileDialogs/Private/SlateFileDlgWindow.h b/Engine/Source/Developer/SlateFileDialogs/Private/SlateFileDlgWindow.h
--- a/Engine/Source/Developer/SlateFileDialogs/Private/SlateFileDlgWindow.h
+++ b/Engine/Source/Developer/SlateFileDialogs/Private/SlateFileDlgWindow.h
@@ -137,7 +137,7 @@ private:
 	void RefreshCrumbs();
 	void OnPathMenuItemClicked( FString ClickedPath );
 	void OnItemSelected(TSharedPtr<FFileEntry> Item, ESelectInfo::Type SelectInfo);
-
+	void ParseTextField(TArray<FString> &FilenameArray, FString Files);
 	void Tick(const FGeometry &AllottedGeometry, const double InCurrentTime, const float InDeltaTime);
 
 	TSharedRef<ITableRow> OnGenerateWidgetForList(TSharedPtr<FFileEntry> Item, const TSharedRef<STableViewBase> &OwnerTable);
@@ -155,6 +155,9 @@ private:
 	bool OnNewDirectoryTextChanged(const FText &InText, FText &ErrorMsg);
 	FReply OnNewDirectoryAcceptCancelClick(FSlateFileDlgWindow::EResult ButtonID);
 
+	FReply OnGoForwardClick();
+	FReply OnGoBackClick();
+
 	/** Collects the output files. */
 	void SetOutputFiles();
 
@@ -209,6 +212,9 @@ private:
 	FDelegateHandle OnDialogDirectoryChangedDelegateHandle;
 	FString RegisteredPath;
 	FString NewDirectoryName;
+
+	TArray<FString> History;
+	int32 HistoryIndex;
 };
 
 
@@ -238,7 +244,84 @@ public:
 	virtual TSharedRef<SWidget> GenerateWidgetForColumn( const FName& ColumnName ) override
 	{
 		FSlateFontInfo ItemFont = StyleSet->GetFontStyle("SlateFileDialogs.Dialog");
-		struct EVisibility FolderIconVisibility = DialogItem->bIsDirectory ? EVisibility::Visible : EVisibility::Hidden;
+		struct EVisibility FolderIconVisibility = EVisibility::Visible;
+		const FSlateBrush *Icon;
+			
+		if (DialogItem->bIsDirectory)
+		{
+			Icon = StyleSet->GetBrush("SlateFileDialogs.Folder16");
+		}
+		else
+		{
+			FString Extension = FPaths::GetExtension(DialogItem->Label, false);
+
+			if (Extension.Equals(TEXT("uasset"), ESearchCase::IgnoreCase))
+			{
+				Icon = StyleSet->GetBrush("SlateFileDialogs.UAsset16");
+			}
+			else if (Extension.Equals(TEXT("uproject"), ESearchCase::IgnoreCase))
+			{
+				Icon = StyleSet->GetBrush("SlateFileDialogs.UProject16");
+			}
+
+			else if (Extension.Equals(TEXT("fbx"), ESearchCase::IgnoreCase))
+			{
+				Icon = StyleSet->GetBrush("SlateFileDialogs.Model3D");
+			}
+
+			else if (Extension.Equals(TEXT("cpp"), ESearchCase::IgnoreCase))
+			{
+				Icon = StyleSet->GetBrush("SlateFileDialogs.TextFile");
+			}
+			else if (Extension.Equals(TEXT("h"), ESearchCase::IgnoreCase))
+			{
+				Icon = StyleSet->GetBrush("SlateFileDialogs.TextFile");
+			}
+			else if (Extension.Equals(TEXT("txt"), ESearchCase::IgnoreCase))
+			{
+				Icon = StyleSet->GetBrush("SlateFileDialogs.TextFile");
+			}
+			else if (Extension.Equals(TEXT("log"), ESearchCase::IgnoreCase))
+			{
+				Icon = StyleSet->GetBrush("SlateFileDialogs.TextFile");
+			}
+
+			else if (Extension.Equals(TEXT("wav"), ESearchCase::IgnoreCase))
+			{
+				Icon = StyleSet->GetBrush("SlateFileDialogs.Audio");
+			}
+			else if (Extension.Equals(TEXT("mp3"), ESearchCase::IgnoreCase))
+			{
+				Icon = StyleSet->GetBrush("SlateFileDialogs.Audio");
+			}
+			else if (Extension.Equals(TEXT("ogg"), ESearchCase::IgnoreCase))
+			{
+				Icon = StyleSet->GetBrush("SlateFileDialogs.Audio");
+			}
+			else if (Extension.Equals(TEXT("mp4"), ESearchCase::IgnoreCase))
+			{
+				Icon = StyleSet->GetBrush("SlateFileDialogs.Video");
+			}
+
+			else if (Extension.Equals(TEXT("png"), ESearchCase::IgnoreCase))
+			{
+				Icon = StyleSet->GetBrush("SlateFileDialogs.Image");
+			}
+			else if (Extension.Equals(TEXT("jpg"), ESearchCase::IgnoreCase))
+			{
+				Icon = StyleSet->GetBrush("SlateFileDialogs.Image");
+			}
+			else if (Extension.Equals(TEXT("bmp"), ESearchCase::IgnoreCase))
+			{
+				Icon = StyleSet->GetBrush("SlateFileDialogs.Image");
+			}
+
+			else
+			{
+				Icon = StyleSet->GetBrush("SlateFileDialogs.PlaceHolder");
+				FolderIconVisibility = EVisibility::Hidden;
+			}
+		}
 
 		if (ColumnName == TEXT("Pathname"))
 		{
@@ -251,7 +334,7 @@ public:
 					.Padding(FMargin(5.0f, 2.0f))
 					[
 						SNew(SImage)
-						.Image(StyleSet->GetBrush("SlateFileDialogs.Folder16"))
+						.Image(Icon)
 						.Visibility(FolderIconVisibility)
 					]
 
diff --git a/Engine/Source/Runtime/Core/Private/GenericPlatform/GenericPlatformProcess.cpp b/Engine/Source/Runtime/Core/Private/GenericPlatform/GenericPlatformProcess.cpp
--- a/Engine/Source/Runtime/Core/Private/GenericPlatform/GenericPlatformProcess.cpp
+++ b/Engine/Source/Runtime/Core/Private/GenericPlatform/GenericPlatformProcess.cpp
@@ -521,3 +521,12 @@ bool FGenericPlatformProcess::Daemonize()
 	UE_LOG(LogHAL, Fatal, TEXT("FGenericPlatformProcess::Daemonize not implemented on this platform"));
 	return false;
 }
+
+bool FGenericPlatformProcess::IsFirstInstance()
+{
+#if !(UE_BUILD_SHIPPING && WITH_EDITOR)
+	return GIsFirstInstance;
+#elif
+	return true;
+#endif
+}
diff --git a/Engine/Source/Runtime/Core/Private/Linux/LinuxApplication.cpp b/Engine/Source/Runtime/Core/Private/Linux/LinuxApplication.cpp
--- a/Engine/Source/Runtime/Core/Private/Linux/LinuxApplication.cpp
+++ b/Engine/Source/Runtime/Core/Private/Linux/LinuxApplication.cpp
@@ -1683,3 +1683,25 @@ void FLinuxApplication::GetWindowPositionInEventLoop(SDL_HWindow NativeWindow, i
 		*y = 0;
 	}
 }
+
+bool FLinuxApplication::IsMouseAttached() const
+{
+	printf("--- check mouse ---\n");
+	int rc;
+	char Mouse[64] = "/sys/class/input/mouse0";
+	int MouseIdx = strlen(Mouse) - 1;
+	strcat(Mouse, "/device/name");
+
+	for (int i=0; i<9; i++)
+	{
+		Mouse[MouseIdx] = '0' + i;
+		if (access(Mouse, F_OK) == 0)
+		{
+			printf("--- check mouse found %d---\n", i);
+			return true;
+		}
+	}
+
+	printf("--- check mouse none found ---\n");
+	return false;
+}
diff --git a/Engine/Source/Runtime/Core/Private/Linux/LinuxPlatformMisc.cpp b/Engine/Source/Runtime/Core/Private/Linux/LinuxPlatformMisc.cpp
--- a/Engine/Source/Runtime/Core/Private/Linux/LinuxPlatformMisc.cpp
+++ b/Engine/Source/Runtime/Core/Private/Linux/LinuxPlatformMisc.cpp
@@ -882,3 +882,88 @@ bool FLinuxPlatformMisc::GetDiskTotalAndFreeSpace(const FString& InPath, uint64&
 	}
 	return (Err == 0);
 }
+
+
+TArray<uint8> FLinuxPlatformMisc::GetMacAddress()
+{
+	char path[64] = "/sys/class/net/eth0";
+	int eth_idx = strlen(path) - 1;
+	char mac_string[64];
+	uint MAC[6] = {0,0,0,0,0,0};
+	TArray<uint8> Result;
+
+	strcat(path, "/address");
+
+	// check to see if eth0-eth9 exist. Use the MAC from the first one we get a hit on
+	for (int i=0; i<=9; i++)
+	{
+		path[eth_idx] = '0' + i;
+
+		int EtherInterface = open(path, O_RDONLY);
+		if (EtherInterface != -1)
+		{
+			ssize_t ReadBytes = read(EtherInterface, mac_string, 64 - 1);
+			close(EtherInterface);
+
+			if (ReadBytes < 0)
+			{
+				continue;
+			}
+
+			mac_string[ReadBytes] = 0;
+
+			if (ReadBytes > 0)
+			{
+				printf("--- MAC Address: %s\n", mac_string);
+				int rc = sscanf(mac_string,"%x:%x:%x:%x:%x:%x", &MAC[0],&MAC[1],&MAC[2],&MAC[3],&MAC[4],&MAC[5]);
+
+				if (rc == 6)
+				{
+					for (int j=0; j<6; j++)
+					{
+						Result.Add(MAC[j]);
+					}
+
+					break;
+				}
+			}
+		}
+	}
+
+	return Result;
+}
+
+
+static int64 LastBatteryCheck = 0;
+static bool bIsOnBattery = false;
+
+bool FLinuxPlatformMisc::IsRunningOnBattery()
+{
+	char scratch[8];
+	FDateTime Time = FDateTime::Now();
+	int64 Seconds = Time.ToUnixTimestamp();
+
+	// don't poll the OS for battery state on every tick. Just do it once every 10 seconds.
+	if (LastBatteryCheck != 0 && (Seconds - LastBatteryCheck) < 10)
+	{
+		return bIsOnBattery;
+	}
+
+	LastBatteryCheck = Seconds;
+	bIsOnBattery = false;
+
+	int State = open("/sys/class/power_supply/ADP0/online", O_RDONLY);
+	if (State != -1)
+	{
+		// found ACAD device. check its state.
+		ssize_t ReadBytes = read(State, scratch, 1);
+		close(State);
+
+		if (ReadBytes > 0)
+		{
+			bIsOnBattery = (scratch[0] == '0');
+		}
+	}
+
+	return bIsOnBattery;
+}
diff --git a/Engine/Source/Runtime/Core/Private/Linux/LinuxPlatformProcess.cpp b/Engine/Source/Runtime/Core/Private/Linux/LinuxPlatformProcess.cpp
--- a/Engine/Source/Runtime/Core/Private/Linux/LinuxPlatformProcess.cpp
+++ b/Engine/Source/Runtime/Core/Private/Linux/LinuxPlatformProcess.cpp
@@ -100,12 +100,7 @@ const TCHAR* FLinuxPlatformProcess::ComputerName()
 
 void FLinuxPlatformProcess::CleanFileCache()
 {
-	bool bShouldCleanShaderWorkingDirectory = true;
-#if !(UE_BUILD_SHIPPING && WITH_EDITOR)
-	// Only clean the shader working directory if we are the first instance, to avoid deleting files in use by other instances
-	// @todo - check if any other instances are running right now
-	bShouldCleanShaderWorkingDirectory = GIsFirstInstance;
-#endif
+	bool bShouldCleanShaderWorkingDirectory = IsFirstInstance();
 
 	if (bShouldCleanShaderWorkingDirectory && !FParse::Param( FCommandLine::Get(), TEXT("Multiprocess")))
 	{
@@ -1126,8 +1121,32 @@ void FLinuxPlatformProcess::LaunchFileInDefaultExternalApplication( const TCHAR*
 
 void FLinuxPlatformProcess::ExploreFolder( const TCHAR* FilePath )
 {
-	// TODO This is broken, not an explore action but should be fine if called on a directory
-	FLinuxPlatformProcess::LaunchFileInDefaultExternalApplication(FilePath, NULL, ELaunchVerb::Edit);
+	struct stat st;
+	TCHAR TruncatedPath[MAX_PATH] = TEXT("");
+	FCString::Strcpy(TruncatedPath, FilePath);
+
+	if (stat(TCHAR_TO_UTF8(FilePath), &st) == 0)
+	{
+		// we just want the directory portion of the path
+		if (!S_ISDIR(st.st_mode))
+		{
+			for (int i=FCString::Strlen(TruncatedPath)-1; i > 0; i--)
+			{
+				if (TruncatedPath[i] == TCHAR('/'))
+				{
+					TruncatedPath[i] = 0;
+					break;
+				}
+			}
+		}
+
+		// launch file manager
+		pid_t pid = fork();
+		if (pid == 0)
+		{
+			exit(execl("/usr/bin/xdg-open", "xdg-open", TCHAR_TO_UTF8(TruncatedPath), (char *)0));
+		}
+	}
 }
 
 /**
@@ -1238,3 +1257,65 @@ FString FLinuxPlatformProcess::FProcEnumInfo::GetName() const
 {
 	return FPaths::GetCleanFilename(GetFullPath());
 }
+
+bool FLinuxPlatformProcess::IsFirstInstance()
+{
+#if !(UE_BUILD_SHIPPING && WITH_EDITOR)
+
+	if (GFileLockDescriptor != -1)
+	{
+		// attempt to set lock
+		struct flock lock;
+
+		lock.l_type = F_WRLCK;
+		lock.l_start = 0;
+		lock.l_whence = SEEK_SET;
+		lock.l_len = 0;
+
+		if (fcntl(GFileLockDescriptor, F_SETLK, &lock) >= 0)
+		{
+			// lock file created and successfully locked by this process.
+			printf("we are the first instance\n");
+			GIsFirstInstance = true;
+		}
+		else
+		{
+			// we were unable to lock file. so some other process beat us to lock file.
+			printf("we are *NOT* the first instance\n");
+			GIsFirstInstance = false;
+		}
+	}
+	else
+	{
+		// Should not get here. It means that we were unable to create or open the lock file.
+		// not much we can do other than provide a default return and spit out an error msg.
+		UE_LOG(LogHAL, Warning, TEXT("No process lock file descriptor"));
+		GIsFirstInstance = true;
+	}
+
+	return GIsFirstInstance;
+#elif
+	return true;
+#endif
+}
+
+void FLinuxPlatformProcess::ReleaseProcessLockFile()
+{
+#if !(UE_BUILD_SHIPPING && WITH_EDITOR)
+	printf("*** releasing process lock file\n");
+
+	if (GFileLockDescriptor != -1)
+	{
+		// attempt to release lock
+		struct flock lock;
+
+		lock.l_type = F_UNLCK;
+		lock.l_start = 0;
+		lock.l_whence = SEEK_SET;
+		lock.l_len = 0;
+
+		fcntl(GFileLockDescriptor, F_SETLK, &lock);
+	}
+#endif
+}
+
diff --git a/Engine/Source/Runtime/Core/Private/Linux/LinuxPlatformSplash.cpp b/Engine/Source/Runtime/Core/Private/Linux/LinuxPlatformSplash.cpp
--- a/Engine/Source/Runtime/Core/Private/Linux/LinuxPlatformSplash.cpp
+++ b/Engine/Source/Runtime/Core/Private/Linux/LinuxPlatformSplash.cpp
@@ -891,9 +891,7 @@ void FLinuxPlatformSplash::Show( )
 	const FText GameName = FText::FromString(FApp::GetGameName());
 
 	const TCHAR* SplashImage = GIsEditor ? ( GameName.IsEmpty() ? TEXT("EdSplashDefault") : TEXT("EdSplash") ) : ( GameName.IsEmpty() ? TEXT("SplashDefault") : TEXT("Splash") );
-	// TODO: add an icon
-	//const TCHAR* IconImage = GIsEditor ? ( GameName.IsEmpty() ? TEXT("EdIconDefault") : TEXT("EdIcon") ) : ( GameName.IsEmpty() ? TEXT("IconDefault") : TEXT("Icon.bmp") );
-	const TCHAR* IconImage = GIsEditor ? ( GameName.IsEmpty() ? TEXT("EdSplashDefault") : TEXT("EdSplash") ) : ( GameName.IsEmpty() ? TEXT("SplashDefault") : TEXT("Splash") );
+	const TCHAR* IconImage = GIsEditor ? ( GameName.IsEmpty() ? TEXT("EdIconDefault") : TEXT("EdIcon") ) : ( GameName.IsEmpty() ? TEXT("IconDefault") : TEXT("Icon.bmp") );
 
 	// make sure a splash was found
 	bool IsCustom;
diff --git a/Engine/Source/Runtime/Core/Private/Linux/LinuxWindow.cpp b/Engine/Source/Runtime/Core/Private/Linux/LinuxWindow.cpp
--- a/Engine/Source/Runtime/Core/Private/Linux/LinuxWindow.cpp
+++ b/Engine/Source/Runtime/Core/Private/Linux/LinuxWindow.cpp
@@ -208,6 +208,52 @@ void FLinuxWindow::Initialize( FLinuxApplication* const Application, const TShar
 	//	The SDL window doesn't need to be reshaped.
 	//	the size of the window you input is the sizeof the client.
 	HWnd = SDL_CreateWindow( TCHAR_TO_ANSI( *Definition->Title ), X, Y, ClientWidth, ClientHeight, WindowStyle  );
+
+	// Try to find an icon for the window
+	const FText GameName = FText::FromString(FApp::GetGameName());
+	FString Iconpath;
+	static SDL_Surface *IconImage = nullptr;
+
+	if (!GameName.IsEmpty())
+	{
+		if (GIsEditor)
+		{
+			Iconpath = FPaths::GameContentDir() / TEXT("Splash/EdIcon.bmp");
+		}
+		else
+		{
+			Iconpath = FPaths::GameContentDir() / TEXT("Splash/Icon.bmp");
+		}
+
+		if (IFileManager::Get().FileSize(*Iconpath) != -1)
+		{
+			IconImage = SDL_LoadBMP(TCHAR_TO_UTF8(*Iconpath));
+		}
+	}
+
+	if (IconImage == nullptr)
+	{
+		// no game specified or there's no custom icons for the game. use default icons.
+		if (GIsEditor)
+		{
+			Iconpath = FPaths::EngineContentDir() / TEXT("Splash/EdIconDefault.bmp");
+		}
+		else
+		{
+			Iconpath = FPaths::EngineContentDir() / TEXT("Splash/IconDefault.bmp");
+		}
+
+		if (IFileManager::Get().FileSize(*Iconpath) != -1)
+		{
+			IconImage = SDL_LoadBMP(TCHAR_TO_UTF8(*Iconpath));
+		}
+	}
+
+	if (IconImage != nullptr)
+	{
+		SDL_SetWindowIcon(HWnd, IconImage);
+	}
+
 	SDL_SetWindowHitTest( HWnd, FLinuxWindow::HitTest, this );
 
 	/* 
diff --git a/Engine/Source/Runtime/Core/Private/Misc/Core.cpp b/Engine/Source/Runtime/Core/Private/Misc/Core.cpp
--- a/Engine/Source/Runtime/Core/Private/Misc/Core.cpp
+++ b/Engine/Source/Runtime/Core/Private/Misc/Core.cpp
@@ -212,6 +212,7 @@ uint32					GFrameNumberRenderThread		= 1;
 // We cannot count on this variable to be accurate in a shipped game, so make sure no code tries to use it
 /** Whether we are the first instance of the game running.													*/
 bool					GIsFirstInstance				= true;
+int32					GFileLockDescriptor				= -1;
 #endif
 /** Threshold for a frame to be considered a hitch (in seconds. */
 float GHitchThreshold = 0.075f;
diff --git a/Engine/Source/Runtime/Core/Public/CoreGlobals.h b/Engine/Source/Runtime/Core/Public/CoreGlobals.h
--- a/Engine/Source/Runtime/Core/Public/CoreGlobals.h
+++ b/Engine/Source/Runtime/Core/Public/CoreGlobals.h
@@ -242,6 +242,7 @@ extern CORE_API uint32 GFrameNumberRenderThread;
 // We cannot count on this variable to be accurate in a shipped game, so make sure no code tries to use it
 /** Whether we are the first instance of the game running. */
 extern CORE_API bool GIsFirstInstance;
+extern CORE_API int32 GFileLockDescriptor;
 #endif
 
 /** Threshold for a frame to be considered a hitch (in seconds. */
diff --git a/Engine/Source/Runtime/Core/Public/GenericPlatform/GenericPlatformProcess.h b/Engine/Source/Runtime/Core/Public/GenericPlatform/GenericPlatformProcess.h
--- a/Engine/Source/Runtime/Core/Public/GenericPlatform/GenericPlatformProcess.h
+++ b/Engine/Source/Runtime/Core/Public/GenericPlatform/GenericPlatformProcess.h
@@ -553,6 +553,12 @@ struct CORE_API FGenericPlatformProcess
 	 * @return true if successful, false otherwise.
 	 */
 	static bool Daemonize();
+
+	/** Check to see if we're the first instance */
+	static bool IsFirstInstance();
+
+	/** Release process lock file (if any) */
+	static void ReleaseProcessLockFile() { }
 };
 
 
diff --git a/Engine/Source/Runtime/Core/Public/Linux/LinuxApplication.h b/Engine/Source/Runtime/Core/Public/Linux/LinuxApplication.h
--- a/Engine/Source/Runtime/Core/Public/Linux/LinuxApplication.h
+++ b/Engine/Source/Runtime/Core/Public/Linux/LinuxApplication.h
@@ -113,6 +113,7 @@ public:
 	 */
 	void GetWindowPositionInEventLoop(SDL_HWindow NativeWindow, int *x, int *y);
 
+	virtual bool IsMouseAttached() const override;
 private:
 
 	FLinuxApplication();
diff --git a/Engine/Source/Runtime/Core/Public/Linux/LinuxPlatformMisc.h b/Engine/Source/Runtime/Core/Public/Linux/LinuxPlatformMisc.h
--- a/Engine/Source/Runtime/Core/Public/Linux/LinuxPlatformMisc.h
+++ b/Engine/Source/Runtime/Core/Public/Linux/LinuxPlatformMisc.h
@@ -25,6 +25,9 @@ struct CORE_API FLinuxPlatformMisc : public FGenericPlatformMisc
 	static class GenericApplication* CreateApplication();
 	static void GetEnvironmentVariable(const TCHAR* VariableName, TCHAR* Result, int32 ResultLength);
 	static void SetEnvironmentVar(const TCHAR* VariableName, const TCHAR* Value);
+	static TArray<uint8> GetMacAddress();
+	static bool IsRunningOnBattery();
+
 #if !UE_BUILD_SHIPPING
 	static bool IsDebuggerPresent();
 	FORCEINLINE static void DebugBreak()
diff --git a/Engine/Source/Runtime/Core/Public/Linux/LinuxPlatformProcess.h b/Engine/Source/Runtime/Core/Public/Linux/LinuxPlatformProcess.h
--- a/Engine/Source/Runtime/Core/Public/Linux/LinuxPlatformProcess.h
+++ b/Engine/Source/Runtime/Core/Public/Linux/LinuxPlatformProcess.h
@@ -280,6 +280,8 @@ struct CORE_API FLinuxPlatformProcess : public FGenericPlatformProcess
 	static bool ExecProcess( const TCHAR* URL, const TCHAR* Params, int32* OutReturnCode, FString* OutStdOut, FString* OutStdErr );
 	static void ExploreFolder( const TCHAR* FilePath );
 	static void LaunchFileInDefaultExternalApplication( const TCHAR* FileName, const TCHAR* Parms = NULL, ELaunchVerb::Type Verb = ELaunchVerb::Open );
+	static bool IsFirstInstance();
+	static void ReleaseProcessLockFile();
 };
 
 typedef FLinuxPlatformProcess FPlatformProcess;
diff --git a/Engine/Source/Runtime/Linux/LinuxCommonStartup/Private/LinuxCommonStartup.cpp b/Engine/Source/Runtime/Linux/LinuxCommonStartup/Private/LinuxCommonStartup.cpp
--- a/Engine/Source/Runtime/Linux/LinuxCommonStartup/Private/LinuxCommonStartup.cpp
+++ b/Engine/Source/Runtime/Linux/LinuxCommonStartup/Private/LinuxCommonStartup.cpp
@@ -186,6 +186,58 @@ int CommonLinuxMain(int argc, char *argv[], int (*RealMain)(const TCHAR * Comman
 		GSavedCommandLine += Temp; 	// note: technically it depends on locale
 	}
 
+#if !(UE_BUILD_SHIPPING && WITH_EDITOR)
+	// check to see if we're first instance and try to set file lock.
+	if ( !FParse::Param( *GSavedCommandLine,TEXT("neverfirst")) )
+	{
+		char LockFileName[128];
+		strcpy(LockFileName,"/tmp/");
+
+		int i;
+		char *ptr = argv[0];
+		for (i = strlen(argv[0])-1; i >= 0 && ptr[i] != '/'; i--);
+
+		strncat(LockFileName, &ptr[i+1], 127);
+		strncat(LockFileName, ".lock", 127);
+		GFileLockDescriptor = open(LockFileName, O_RDWR | O_CREAT, 0666);
+
+		if (GFileLockDescriptor != -1)
+		{
+			// attempt to set lock
+			struct flock lock;
+
+			lock.l_type = F_WRLCK;
+			lock.l_start = 0;
+			lock.l_whence = SEEK_SET;
+			lock.l_len = 0;
+
+			if (fcntl(GFileLockDescriptor, F_SETLK, &lock) >= 0)
+			{
+				// lock file created and successfully locked by this process.
+				printf("** we are the first instance\n");
+				GIsFirstInstance = true;
+			}
+			else
+			{
+				// we were unable to lock file. so some other process beat us to lock file.
+				printf("** we are *NOT* the first instance\n");
+				GIsFirstInstance = false;
+			}
+		}
+		else
+		{
+			// Should not get here. It means that we were unable to create or open the lock file.
+			// not much we can do other than provide a default return and spit out an error msg.
+			fprintf(stderr, "*** LaunchLinux - unable to create process lock file '%s'. (%s)\n", LockFileName, strerror(errno));
+			GIsFirstInstance = true;
+		}
+	}
+	else
+	{
+		GIsFirstInstance = false;
+	}
+#endif
+
 	if (!UE_BUILD_SHIPPING)
 	{
 		GAlwaysReportCrash = true;	// set by default and reverse the behavior
@@ -222,6 +274,13 @@ int CommonLinuxMain(int argc, char *argv[], int (*RealMain)(const TCHAR * Comman
 		}
 	}
 
+#if !(UE_BUILD_SHIPPING && WITH_EDITOR)
+	if (GFileLockDescriptor >= 0)
+	{
+		close(GFileLockDescriptor);
+	}
+#endif
+
 	if (ErrorLevel)
 	{
 		printf("Exiting abnormally (error code: %d)\n", ErrorLevel);
