SlateFileDialogs: minor fixes.

- Includes PR #1548: Fixed:index out of range (Contributed by [REDACTED_PERSON])

diff --git a/Engine/Source/Developer/SlateFileDialogs/Private/SlateFileDlgWindow.cpp b/Engine/Source/Developer/SlateFileDialogs/Private/SlateFileDlgWindow.cpp
--- a/Engine/Source/Developer/SlateFileDialogs/Private/SlateFileDlgWindow.cpp
+++ b/Engine/Source/Developer/SlateFileDialogs/Private/SlateFileDlgWindow.cpp
@@ -853,18 +853,22 @@ void SSlateFileOpenDlg::BuildDirectoryPath()
 			BuiltPath = BuiltPath + TEXT("/") + AbsPath.Left(Idx);
 		}
 
-		FCString::Strcpy(Temp, AbsPath.Len(), &AbsPath[Idx+1]);
+		FCString::Strcpy(Temp, ARRAY_COUNT(Temp), &AbsPath[Idx < AbsPath.Len() - 1 ? Idx + 1 : Idx]);
 
 		DirectoryNodesArray.Add(FDirNode(AbsPath.Left(Idx), nullptr));
 	}
 	else if (PLATFORM_LINUX)
 	{
 		// start with system base directory
-		FCString::Strcpy(Temp, AbsPath.Len(), *AbsPath);
+		FCString::Strcpy(Temp, ARRAY_COUNT(Temp), *AbsPath);
 
 		BuiltPath = "/";
 		DirectoryNodesArray.Add(FDirNode(FString(TEXT("/")), nullptr));
 	}
+	else
+	{
+		checkf(false, TEXT("SlateDialogs will not work on this platform (modify SSlateFileOpenDlg::BuildDirectoryPath())"));
+	}
 
 	// break path into tokens
 	TCHAR *ContextStr = nullptr;
