Slate dialog modal window was not settings its parent window

diff --git a/Engine/Source/Developer/SlateFileDialogs/Private/SlateFileDlgWindow.cpp b/Engine/Source/Developer/SlateFileDialogs/Private/SlateFileDlgWindow.cpp
--- a/Engine/Source/Developer/SlateFileDialogs/Private/SlateFileDlgWindow.cpp
+++ b/Engine/Source/Developer/SlateFileDialogs/Private/SlateFileDlgWindow.cpp
@@ -11,6 +11,7 @@
 #include "Framework/Application/SlateApplication.h"
 #include "Textures/SlateIcon.h"
 #include "Framework/Commands/UIAction.h"
+#include "Framework/Docking/TabManager.h"
 #include "Widgets/Layout/SSpacer.h"
 #include "Framework/MultiBox/MultiBoxBuilder.h"
 #include "Widgets/Input/SButton.h"
@@ -231,7 +232,9 @@ bool FSlateFileDlgWindow::OpenFileDialog(const void* ParentWindowHandle, const F
 	
 	ModalWindow->SetContent( DialogWidget.ToSharedRef() );
 		
-	FSlateApplication::Get().AddModalWindow(ModalWindow, NULL);
+	TSharedPtr<SWindow> RootWindow = FGlobalTabmanager::Get()->GetRootWindow();
+	FSlateApplication::Get().AddModalWindow(ModalWindow, RootWindow);
+
 	return (DialogWidget->GetResponse() == EResult::Accept && OutFilenames.Num() > 0);
 }
 
@@ -277,7 +280,9 @@ bool FSlateFileDlgWindow::OpenDirectoryDialog(const void* ParentWindowHandle, co
 
 	ModalWindow->SetContent( DialogWidget.ToSharedRef() );
 
-	FSlateApplication::Get().AddModalWindow(ModalWindow, NULL);
+	TSharedPtr<SWindow> RootWindow = FGlobalTabmanager::Get()->GetRootWindow();
+	FSlateApplication::Get().AddModalWindow(ModalWindow, RootWindow);
+
 	bool RC = (DialogWidget->GetResponse() == EResult::Accept && TempOut.Num() > 0);
 
 	if (TempOut.Num() > 0)
@@ -327,7 +332,9 @@ bool FSlateFileDlgWindow::SaveFileDialog(const void* ParentWindowHandle, const F
 
 	ModalWindow->SetContent( DialogWidget.ToSharedRef() );
 		
-	FSlateApplication::Get().AddModalWindow(ModalWindow, NULL);
+	TSharedPtr<SWindow> RootWindow = FGlobalTabmanager::Get()->GetRootWindow();
+	FSlateApplication::Get().AddModalWindow(ModalWindow, RootWindow);
+
 	return (DialogWidget->GetResponse() == EResult::Accept && OutFilenames.Num() > 0);
 }
 
