Fix non-unity

diff --git a/Engine/Source/Developer/SlateFileDialogs/SlateFileDialogs.Build.cs b/Engine/Source/Developer/SlateFileDialogs/SlateFileDialogs.Build.cs
--- a/Engine/Source/Developer/SlateFileDialogs/SlateFileDialogs.Build.cs
+++ b/Engine/Source/Developer/SlateFileDialogs/SlateFileDialogs.Build.cs
@@ -9,6 +9,7 @@ public class SlateFileDialogs : ModuleRules
         PrivateDependencyModuleNames.AddRange(
             new string[] {
                 "Core",
+                "CoreUObject",
                 "InputCore",
                 "Slate",
                 "SlateCore",
