SlateFileDlgWindow now unregisters correctly from DirectoryWatcher

(fix by brandon.schaefer)

diff --git a/Engine/Source/Developer/SlateFileDialogs/Private/SlateFileDlgWindow.cpp b/Engine/Source/Developer/SlateFileDialogs/Private/SlateFileDlgWindow.cpp
--- a/Engine/Source/Developer/SlateFileDialogs/Private/SlateFileDlgWindow.cpp
+++ b/Engine/Source/Developer/SlateFileDialogs/Private/SlateFileDlgWindow.cpp
@@ -845,6 +845,13 @@ void SSlateFileOpenDlg::Construct(const FArguments& InArgs)
 }
 END_SLATE_FUNCTION_BUILD_OPTIMIZATION
 
+SSlateFileOpenDlg::~SSlateFileOpenDlg()
+{
+	if (DirectoryWatcher && RegisteredPath.Len() > 0)
+	{
+		DirectoryWatcher->UnregisterDirectoryChangedCallback_Handle(RegisteredPath, OnDialogDirectoryChangedDelegateHandle);
+	}
+}
 
 void SSlateFileOpenDlg::BuildDirectoryPath()
 {
diff --git a/Engine/Source/Developer/SlateFileDialogs/Private/SlateFileDlgWindow.h b/Engine/Source/Developer/SlateFileDialogs/Private/SlateFileDlgWindow.h
--- a/Engine/Source/Developer/SlateFileDialogs/Private/SlateFileDlgWindow.h
+++ b/Engine/Source/Developer/SlateFileDialogs/Private/SlateFileDlgWindow.h
@@ -141,6 +141,8 @@ public:
 	SLATE_ARGUMENT(FSlateFileDialogsStyle*, StyleSet)
 	SLATE_END_ARGS()
 
+	~SSlateFileOpenDlg();
+
 	void Construct(const FArguments& InArgs);
 	FSlateFileDlgWindow::EResult GetResponse() { return UserResponse; }
 	void SetOutNames(TArray<FString>* Ptr) { OutNames = Ptr; }
