Linux: SlateFileDialog has filters that end up being garabge from a missing second null terminator when adding to a list

diff --git a/Engine/Source/Developer/SlateFileDialogs/Private/SlateFileDlgWindow.cpp b/Engine/Source/Developer/SlateFileDialogs/Private/SlateFileDlgWindow.cpp
--- a/Engine/Source/Developer/SlateFileDialogs/Private/SlateFileDlgWindow.cpp
+++ b/Engine/Source/Developer/SlateFileDialogs/Private/SlateFileDlgWindow.cpp
@@ -1497,7 +1497,7 @@ void SSlateFileOpenDlg::ParseFilters()
 		{
 			if (FilterNameArray.Num() == 0)
 			{
-				TCHAR Temp[MAX_FILTER_LENGTH];
+				TCHAR Temp[MAX_FILTER_LENGTH] = {0};
 				FCString::Strcpy(Temp, ARRAY_COUNT(Temp), *Filters);
 
 				// break path into tokens
@@ -1541,7 +1541,7 @@ bool SSlateFileOpenDlg::GetFilterExtension(FString &OutString)
 	}
 
 	// make a copy of filter string that we can modify
-	TCHAR Temp[MAX_FILTER_LENGTH];
+	TCHAR Temp[MAX_FILTER_LENGTH] = {0};
 	FCString::Strcpy(Temp, ARRAY_COUNT(Temp), *(*FilterNameArray[FilterIndex].Get()));
 
 	// find start of extension
