Fix crash when clicking on empty space in the drop down menu in slage file dialog window

diff --git a/Engine/Source/Developer/SlateFileDialogs/Private/SlateFileDlgWindow.cpp b/Engine/Source/Developer/SlateFileDialogs/Private/SlateFileDlgWindow.cpp
--- a/Engine/Source/Developer/SlateFileDialogs/Private/SlateFileDlgWindow.cpp
+++ b/Engine/Source/Developer/SlateFileDialogs/Private/SlateFileDlgWindow.cpp
@@ -1334,7 +1334,7 @@ void SSlateFileOpenDlg::OnFilterChanged(TSharedPtr<FString> NewValue, ESelectInf
 {
 	for (int32 i = 0; i < FilterNameArray.Num(); i++)
 	{
-		if (FilterNameArray[i].Get()->Compare(*NewValue.Get(), ESearchCase::CaseSensitive) == 0)
+		if (NewValue.IsValid() && FilterNameArray[i].Get()->Compare(*NewValue.Get(), ESearchCase::CaseSensitive) == 0)
 		{
 			FilterIndex = i;
 			break;
