Always put an extension at the end of a file if it exists
If we currently have a SaveFilename keep that name around when going in and out of folders

diff --git a/Engine/Source/Developer/SlateFileDialogs/Private/SlateFileDlgWindow.cpp b/Engine/Source/Developer/SlateFileDialogs/Private/SlateFileDlgWindow.cpp
--- a/Engine/Source/Developer/SlateFileDialogs/Private/SlateFileDlgWindow.cpp
+++ b/Engine/Source/Developer/SlateFileDialogs/Private/SlateFileDlgWindow.cpp
@@ -1304,8 +1304,6 @@ void SSlateFileOpenDlg::OnItemDoubleClicked(TSharedPtr<FFileEntry> Item)
 {
 	if (Item->bIsDirectory)
 	{
-		SetDefaultFile(FString(""));
-
 		CurrentPath = CurrentPath + Item->Label + TEXT("/");
 		bNeedsBuilding = true;
 		bRebuildDirPath = true;
@@ -1422,19 +1420,9 @@ void SSlateFileOpenDlg::SetDefaultFile(FString DefaultFile)
 	else
 	{
 		SaveFilename = FileList;
-	}
-
-	SaveFilenameEditBox->SetText(SaveFilename);
-}
-
 
-void SSlateFileOpenDlg::OnFileNameCommitted(const FText& InText, ETextCommit::Type InCommitType)
-{
-	// update edit box unless user choose to escape out
-	if (InCommitType != ETextCommit::OnCleared)
-	{
+		// ensure we always place an extension on the SaveFilename
 		FString Extension;
-		SaveFilename = InText.ToString();
 
 		// get current filter extension
 		if (!bDirectoriesOnly && GetFilterExtension(Extension))
@@ -1446,7 +1434,18 @@ void SSlateFileOpenDlg::OnFileNameCommitted(const FText& InText, ETextCommit::Ty
 				SaveFilename = SaveFilename + Extension;
 			}
 		}
+	}
 
+	SaveFilenameEditBox->SetText(SaveFilename);
+}
+
+
+void SSlateFileOpenDlg::OnFileNameCommitted(const FText& InText, ETextCommit::Type InCommitType)
+{
+	// update edit box unless user choose to escape out
+	if (InCommitType != ETextCommit::OnCleared)
+	{
+		SaveFilename = InText.ToString();
 		ListView->ClearSelection();
 
 		SetDefaultFile(SaveFilename);
@@ -1540,6 +1539,12 @@ bool SSlateFileOpenDlg::GetFilterExtension(FString &OutString)
 		return false;
 	}
 
+	// We have attempted to get the filter extension before parsing them
+	if (FilterNameArray.Num() == 0)
+	{
+		ParseFilters();
+	}
+
 	// make a copy of filter string that we can modify
 	TCHAR Temp[MAX_FILTER_LENGTH] = {0};
 	FCString::Strcpy(Temp, ARRAY_COUNT(Temp), *(*FilterNameArray[FilterIndex].Get()));
@@ -1645,8 +1650,6 @@ FReply SSlateFileOpenDlg::OnGoForwardClick()
 {
 	if ((HistoryIndex+1) < History.Num())
 	{
-		SetDefaultFile(FString(""));
-
 		HistoryIndex++;
 		CurrentPath = History[HistoryIndex];
 		bNeedsBuilding = true;
@@ -1662,8 +1665,6 @@ FReply SSlateFileOpenDlg::OnGoBackClick()
 {
 	if (HistoryIndex > 0)
 	{
-		SetDefaultFile(FString(""));
-
 		HistoryIndex--;
 		CurrentPath = History[HistoryIndex];
 		bNeedsBuilding = true;
