Better fix for extensions export/import open/save diaglog

diff --git a/Engine/Source/Developer/SlateFileDialogs/Private/SlateFileDlgWindow.cpp b/Engine/Source/Developer/SlateFileDialogs/Private/SlateFileDlgWindow.cpp
--- a/Engine/Source/Developer/SlateFileDialogs/Private/SlateFileDlgWindow.cpp
+++ b/Engine/Source/Developer/SlateFileDialogs/Private/SlateFileDlgWindow.cpp
@@ -755,7 +755,7 @@ void SSlateFileOpenDlg::Construct(const FArguments& InArgs)
 						.Padding(FMargin(0.0f))
 						[
 							SNew(STextBlock)
-							.Text(LOCTEXT("FilterLabel", "Filter:"))
+							.Text(bSaveFile ? LOCTEXT("SaveTypeLabel", "Save as type:") : LOCTEXT("FilterLabel", "Filter:"))
 							.Font(StyleSet->GetFontStyle("SlateFileDialogs.Dialog"))
 							.Justification(ETextJustify::Left)
 						]
@@ -1304,6 +1304,11 @@ void SSlateFileOpenDlg::OnItemDoubleClicked(TSharedPtr<FFileEntry> Item)
 {
 	if (Item->bIsDirectory)
 	{
+		if (!bSaveFile)
+		{
+			SetDefaultFile(FString(""));
+		}
+
 		CurrentPath = CurrentPath + Item->Label + TEXT("/");
 		bNeedsBuilding = true;
 		bRebuildDirPath = true;
@@ -1390,6 +1395,19 @@ void SSlateFileOpenDlg::ParseTextField(TArray<FString> &FilenameArray, FString F
 	}
 	else
 	{
+		FString Extension;
+
+		// get current filter extension
+		if (!bDirectoriesOnly && GetFilterExtension(Extension))
+		{
+			// append extension to filename if user left it off
+			if (!SaveFilename.EndsWith(Extension, ESearchCase::CaseSensitive) &&
+				!IsWildcardExtension(Extension))
+			{
+				Files = Files + Extension;
+			}
+		}
+
 		FilenameArray.Add(Files);
 	}
 }
@@ -1420,20 +1438,6 @@ void SSlateFileOpenDlg::SetDefaultFile(FString DefaultFile)
 	else
 	{
 		SaveFilename = FileList;
-
-		// ensure we always place an extension on the SaveFilename
-		FString Extension;
-
-		// get current filter extension
-		if (!bDirectoriesOnly && GetFilterExtension(Extension))
-		{
-			// append extension to filename if user left it off
-			if (!SaveFilename.EndsWith(Extension, ESearchCase::CaseSensitive) &&
-				!IsWildcardExtension(Extension))
-			{
-				SaveFilename = SaveFilename + Extension;
-			}
-		}
 	}
 
 	SaveFilenameEditBox->SetText(SaveFilename);
@@ -1650,6 +1654,11 @@ FReply SSlateFileOpenDlg::OnGoForwardClick()
 {
 	if ((HistoryIndex+1) < History.Num())
 	{
+		if (!bSaveFile)
+		{
+			SetDefaultFile(FString(""));
+		}
+
 		HistoryIndex++;
 		CurrentPath = History[HistoryIndex];
 		bNeedsBuilding = true;
@@ -1665,6 +1674,11 @@ FReply SSlateFileOpenDlg::OnGoBackClick()
 {
 	if (HistoryIndex > 0)
 	{
+		if (!bSaveFile)
+		{
+			SetDefaultFile(FString(""));
+		}
+
 		HistoryIndex--;
 		CurrentPath = History[HistoryIndex];
 		bNeedsBuilding = true;
