Slate save dialog - Add support for saving a file when the file format supports multiple extensions.

diff --git a/Engine/Source/Developer/SlateFileDialogs/Private/SlateFileDlgWindow.cpp b/Engine/Source/Developer/SlateFileDialogs/Private/SlateFileDlgWindow.cpp
--- a/Engine/Source/Developer/SlateFileDialogs/Private/SlateFileDlgWindow.cpp
+++ b/Engine/Source/Developer/SlateFileDialogs/Private/SlateFileDlgWindow.cpp
@@ -1395,16 +1395,33 @@ void SSlateFileOpenDlg::ParseTextField(TArray<FString> &FilenameArray, FString F
 	}
 	else
 	{
-		FString Extension;
+		TArray<FString> Extensions;
+		FString FirstNonWilcardExtension;
 
 		// get current filter extension
-		if (!bDirectoriesOnly && GetFilterExtension(Extension))
+		if (!bDirectoriesOnly && GetFilterExtensions(Extensions))
 		{
+			bool bSaveFilenameHasExtension = false;
+
+			for ( const FString& Extension : Extensions )
+			{
+				bSaveFilenameHasExtension = ( !IsWildcardExtension(Extension) && SaveFilename.EndsWith(Extension, ESearchCase::CaseSensitive) );
+
+				if ( !IsWildcardExtension(Extension) && FirstNonWilcardExtension.IsEmpty() )
+				{
+					FirstNonWilcardExtension = Extension;
+				}
+
+				if ( bSaveFilenameHasExtension )
+				{
+					break;
+				}
+			}
+
 			// append extension to filename if user left it off
-			if (!SaveFilename.EndsWith(Extension, ESearchCase::CaseSensitive) &&
-				!IsWildcardExtension(Extension))
+			if (!bSaveFilenameHasExtension && !FirstNonWilcardExtension.IsEmpty())
 			{
-				Files = Files + Extension;
+				Files = Files + FirstNonWilcardExtension;
 			}
 		}
 
@@ -1533,9 +1550,9 @@ void SSlateFileOpenDlg::ParseFilters()
 }
 
 
-bool SSlateFileOpenDlg::GetFilterExtension(FString &OutString)
+bool SSlateFileOpenDlg::GetFilterExtensions(TArray<FString>& OutExtensions)
 {
-	OutString.Empty();
+	OutExtensions.Reset();
 
 	// check to see if filters were given
 	if (Filters.Len() == 0)
@@ -1544,38 +1561,51 @@ bool SSlateFileOpenDlg::GetFilterExtension(FString &OutString)
 	}
 
 	// We have attempted to get the filter extension before parsing them
-	if (FilterNameArray.Num() == 0)
+	if (FilterListArray.Num() == 0)
 	{
 		ParseFilters();
 	}
 
-	// make a copy of filter string that we can modify
-	TCHAR Temp[MAX_FILTER_LENGTH] = {0};
-	FCString::Strcpy(Temp, UE_ARRAY_COUNT(Temp), *(*FilterNameArray[FilterIndex].Get()));
+	if (!FilterListArray.IsValidIndex(FilterIndex))
+	{
+		return false;
+	}
 
-	// find start of extension
-	TCHAR *FilterExt = FCString::Strchr(Temp, '.');
-	if (FilterExt != nullptr)
+	TArray<FString> Extensions;
+	const bool bCullEmpty = true;
+	if (FilterListArray[FilterIndex].ParseIntoArray(Extensions, TEXT(";"), bCullEmpty) > 0)
 	{
-		// strip any trailing junk
-		int32 i;
-		for (i = 0; i < FCString::Strlen(FilterExt); i++)
+		for (FString& Extension : Extensions)
 		{
-			if (FilterExt[i] == ' ' || FilterExt[i] == ')' || FilterExt[i] == ';')
+			// find start of extension
+			int32 DotIndex;
+			if (Extension.FindChar(TEXT('.'), DotIndex))
 			{
-				FilterExt[i] = 0;
-				break;
+				Extension.RightChopInline(DotIndex);
+
+				// strip any trailing junk
+				for (TCHAR& ExtensionChar : Extension)
+				{
+					if (ExtensionChar == TEXT(' ') || ExtensionChar == TEXT(')') || ExtensionChar == TEXT(';'))
+					{
+						ExtensionChar = 0;
+						break;
+					}
+				}
+
+				Extension.TrimToNullTerminator();
+
+				// store result and clean up
+				OutExtensions.Add(Extension);
+			}
+			else if (Extension[0] == TEXT('*'))
+			{
+				OutExtensions.Add(Extension);
 			}
 		}
-
-		// store result and clean up
-		OutString = FilterExt;
 	}
-	else if (Temp[0] == TEXT('*'))
-	{
-		OutString = Temp;
-	}
-	return !OutString.IsEmpty();
+
+	return OutExtensions.Num() > 0;
 }
 
 bool SSlateFileOpenDlg::IsWildcardExtension(const FString& Extension)
diff --git a/Engine/Source/Developer/SlateFileDialogs/Private/SlateFileDlgWindow.h b/Engine/Source/Developer/SlateFileDialogs/Private/SlateFileDlgWindow.h
--- a/Engine/Source/Developer/SlateFileDialogs/Private/SlateFileDlgWindow.h
+++ b/Engine/Source/Developer/SlateFileDialogs/Private/SlateFileDlgWindow.h
@@ -183,7 +183,7 @@ private:
 	/** Collects the output files. */
 	void SetOutputFiles();
 
-	bool GetFilterExtension(FString &OutString);
+	bool GetFilterExtensions(TArray< FString >& OutExtensions);
 	void ParseFilters();
 	/** @return true if the extension filter contains a wildcard or not */
 	bool IsWildcardExtension(const FString& Extension);
