Remove unused private fields - Linux

diff --git a/Engine/Source/Developer/SlateFileDialogs/Public/ISlateFileDialogModule.h b/Engine/Source/Developer/SlateFileDialogs/Public/ISlateFileDialogModule.h
--- a/Engine/Source/Developer/SlateFileDialogs/Public/ISlateFileDialogModule.h
+++ b/Engine/Source/Developer/SlateFileDialogs/Public/ISlateFileDialogModule.h
@@ -41,7 +41,4 @@ public:
 	}
 
 	virtual ISlateFileDialogsModule* Get() { return nullptr; }
-
-private:
-	ISlateFileDialogsModule *SlateFileDialog;
 };
diff --git a/Engine/Source/Runtime/ApplicationCore/Public/Linux/LinuxCursor.h b/Engine/Source/Runtime/ApplicationCore/Public/Linux/LinuxCursor.h
--- a/Engine/Source/Runtime/ApplicationCore/Public/Linux/LinuxCursor.h
+++ b/Engine/Source/Runtime/ApplicationCore/Public/Linux/LinuxCursor.h
@@ -74,8 +74,6 @@ private:
 
 	SDL_Rect CursorClipRect;
 
-	uint32 CursorEvent;
-
 	/** Cached global X position */
 	mutable int32 CachedGlobalXPosition;
 
