Converted FilterCombo from a STextComboBox to a SComboBox allowing us to override the width of the dropdown to cut off long strings

diff --git a/Engine/Source/Developer/SlateFileDialogs/Private/SlateFileDlgWindow.cpp b/Engine/Source/Developer/SlateFileDialogs/Private/SlateFileDlgWindow.cpp
--- a/Engine/Source/Developer/SlateFileDialogs/Private/SlateFileDlgWindow.cpp
+++ b/Engine/Source/Developer/SlateFileDialogs/Private/SlateFileDlgWindow.cpp
@@ -770,12 +770,23 @@ void SSlateFileOpenDlg::Construct(const FArguments& InArgs)
 							.MaxDesiredWidth(200.0f)
 							.Padding(FMargin(0.0f))
 							[
-								SAssignNew(FilterCombo, STextComboBox)
+								SAssignNew(FilterCombo, SComboBox<TSharedPtr<FString>>)
 								.ContentPadding(FMargin(4.0f, 2.0f))
-								//.MaxListHeight(100.0f)
 								.OptionsSource(&FilterNameArray)
-								.Font(StyleSet->GetFontStyle("SlateFileDialogs.Dialog"))
+								.OnGenerateWidget_Lambda([this](TSharedPtr<FString> Item)
+								{ 
+									return SNew(SBox)
+										.MaxDesiredWidth(600.0f)
+										[
+											SNew(STextBlock)
+											.Text(FText::FromString(*Item))
+											.Font(StyleSet->GetFontStyle("SlateFileDialogs.Dialog"))
+										];
+								} )	
 								.OnSelectionChanged(this, &SSlateFileOpenDlg::OnFilterChanged)
+								[
+									SAssignNew(FilterComboBoxTitleBlock, STextBlock)
+								]
 							]
 						]
 					]
@@ -1538,7 +1549,7 @@ void SSlateFileOpenDlg::ParseFilters()
 				}
 			}
 
-			FilterCombo->SetSelectedItem(FilterNameArray[FilterIndex]);
+			FilterComboBoxTitleBlock->SetText(FText::FromString(*FilterNameArray[FilterIndex]));
 		}
 		else
 		{
diff --git a/Engine/Source/Developer/SlateFileDialogs/Private/SlateFileDlgWindow.h b/Engine/Source/Developer/SlateFileDialogs/Private/SlateFileDlgWindow.h
--- a/Engine/Source/Developer/SlateFileDialogs/Private/SlateFileDlgWindow.h
+++ b/Engine/Source/Developer/SlateFileDialogs/Private/SlateFileDlgWindow.h
@@ -193,7 +193,8 @@ private:
 	TArray<TSharedPtr<FFileEntry>> FilesArray;
 	TArray<TSharedPtr<FFileEntry>> LineItemArray;	
 
-	TSharedPtr<STextComboBox> FilterCombo;
+	TSharedPtr<SComboBox<TSharedPtr<FString>>> FilterCombo;
+	TSharedPtr<STextBlock> FilterComboBoxTitleBlock;
 	TSharedPtr<SHorizontalBox> FilterHBox;
 	TSharedPtr<SInlineEditableTextBlock> SaveFilenameEditBox;
 	TSharedPtr<SInlineEditableTextBlock> NewDirectoryEditBox;
