Folders are now properly selected when clicked on in file dialog windows

diff --git a/Engine/Source/Developer/SlateFileDialogs/Private/SlateFileDlgWindow.cpp b/Engine/Source/Developer/SlateFileDialogs/Private/SlateFileDlgWindow.cpp
--- a/Engine/Source/Developer/SlateFileDialogs/Private/SlateFileDlgWindow.cpp
+++ b/Engine/Source/Developer/SlateFileDialogs/Private/SlateFileDlgWindow.cpp
@@ -1496,11 +1496,7 @@ void SSlateFileOpenDlg::OnItemSelected(TSharedPtr<FFileEntry> Item, ESelectInfo:
 			
 			for (int32 i = 0; i < SelectedItems.Num(); i++)
 			{
-				if (SelectedItems[i]->bIsDirectory)
-				{
-					ListView->SetItemSelection(SelectedItems[i], false, ESelectInfo::Direct);
-				}
-				else
+				if (!SelectedItems[i]->bIsDirectory)
 				{
 					FileList = FileList + TEXT("\"") + SelectedItems[i]->Label + TEXT("\" ");
 				}
@@ -1510,11 +1506,8 @@ void SSlateFileOpenDlg::OnItemSelected(TSharedPtr<FFileEntry> Item, ESelectInfo:
 		{
 			FileList = Item->Label;
 		}
-	
-		if (bDirectoriesOnly == Item->bIsDirectory)
-		{
-			SetDefaultFile(FileList);
-		}
+
+		SetDefaultFile(FileList);
 	}
 }
 
