In Slate File Dialogs for opening one or multiple files, clicking open when a directory is highlighted now navigates into that directory.

diff --git a/Engine/Source/Developer/SlateFileDialogs/Private/SlateFileDlgWindow.cpp b/Engine/Source/Developer/SlateFileDialogs/Private/SlateFileDlgWindow.cpp
--- a/Engine/Source/Developer/SlateFileDialogs/Private/SlateFileDlgWindow.cpp
+++ b/Engine/Source/Developer/SlateFileDialogs/Private/SlateFileDlgWindow.cpp
@@ -1143,22 +1143,53 @@ void SSlateFileOpenDlg::SetOutputFiles()
 	}
 }
 
+TSharedPtr<FFileEntry> SSlateFileOpenDlg::GetSoloDirectorySelected() const
+{
+	TArray<TSharedPtr<FFileEntry>> SelectedItems = ListView->GetSelectedItems();
+	if (SelectedItems.Num() == 1 && SelectedItems[0]->bIsDirectory)
+	{
+		return SelectedItems[0];
+	}
+
+	return nullptr;
+}
 
 bool SSlateFileOpenDlg::IsAcceptEnabled() const
 {
-	if (!bDirectoriesOnly)
+	if (bDirectoriesOnly)
 	{
-		return !SaveFilename.IsEmpty();
+		return true;
 	}
 
-	return true;
-}
+	TSharedPtr<FFileEntry> SoloSelectedDirectory = GetSoloDirectorySelected();
 
+	if (SoloSelectedDirectory.IsValid()) 
+	{
+		return true;
+	}
+	else if (!SaveFilename.IsEmpty())
+	{
+		return true;
+	}
+
+	return false;
+}
 
 FReply SSlateFileOpenDlg::OnAcceptCancelClick(FSlateFileDlgWindow::EResult ButtonID)
 {
 	if (ButtonID == FSlateFileDlgWindow::Accept)
 	{
+		if (!bDirectoriesOnly)
+		{
+			TSharedPtr<FFileEntry> SoloSelectedDirectory = GetSoloDirectorySelected();
+
+			if (SoloSelectedDirectory.IsValid())
+			{
+				OnItemDoubleClicked(SoloSelectedDirectory);
+				return FReply::Handled();
+			}
+		}
+
 		SetOutputFiles();
 	}
 	else
@@ -1175,7 +1206,6 @@ FReply SSlateFileOpenDlg::OnAcceptCancelClick(FSlateFileDlgWindow::EResult Butto
 	return FReply::Handled();
 }
 
-
 FReply SSlateFileOpenDlg::OnDirSublevelClick(int32 Level)
 {
 	DirectoryNodesArray[DirNodeIndex].TextBlock->SetFont(StyleSet->GetFontStyle("SlateFileDialogs.Dialog"));
@@ -1197,8 +1227,6 @@ FReply SSlateFileOpenDlg::OnDirSublevelClick(int32 Level)
 	return FReply::Handled();
 }
 
-
-
 void SSlateFileOpenDlg::Tick(const FGeometry &AllottedGeometry, const double InCurrentTime, const float InDeltaTime)
 {
 	SCompoundWidget::Tick(AllottedGeometry, InCurrentTime, InDeltaTime);
@@ -1493,7 +1521,7 @@ void SSlateFileOpenDlg::OnItemSelected(TSharedPtr<FFileEntry> Item, ESelectInfo:
 		if (!bDirectoriesOnly)
 		{
 			TArray<TSharedPtr<FFileEntry>> SelectedItems = ListView->GetSelectedItems();
-			
+
 			for (int32 i = 0; i < SelectedItems.Num(); i++)
 			{
 				if (!SelectedItems[i]->bIsDirectory)
diff --git a/Engine/Source/Developer/SlateFileDialogs/Private/SlateFileDlgWindow.h b/Engine/Source/Developer/SlateFileDialogs/Private/SlateFileDlgWindow.h
--- a/Engine/Source/Developer/SlateFileDialogs/Private/SlateFileDlgWindow.h
+++ b/Engine/Source/Developer/SlateFileDialogs/Private/SlateFileDlgWindow.h
@@ -188,6 +188,9 @@ private:
 	/** @return true if the extension filter contains a wildcard or not */
 	bool IsWildcardExtension(const FString& Extension);
 
+	/** @return ptr to directory if only one directory and zero files are selected, else nullptr */
+	TSharedPtr<FFileEntry> GetSoloDirectorySelected() const;
+
 	TArray< FDirNode > DirectoryNodesArray;
 	TArray<TSharedPtr<FFileEntry>> FoldersArray;
 	TArray<TSharedPtr<FFileEntry>> FilesArray;
