Prevents Slate file dialogs from erroneously trimming an extra layer off of paths.

diff --git a/Engine/Source/Developer/SlateFileDialogs/Private/SlateFileDlgWindow.cpp b/Engine/Source/Developer/SlateFileDialogs/Private/SlateFileDlgWindow.cpp
--- a/Engine/Source/Developer/SlateFileDialogs/Private/SlateFileDlgWindow.cpp
+++ b/Engine/Source/Developer/SlateFileDialogs/Private/SlateFileDlgWindow.cpp
@@ -207,7 +207,7 @@ bool FSlateFileDlgWindow::OpenFileDialog(const void* ParentWindowHandle, const F
 		const FString& DefaultFile, const FString& FileTypes, uint32 Flags, TArray<FString>& OutFilenames, int32& OutFilterIndex)
 {
 	FString StartDirectory = DefaultPath;
-	TrimStartDirectory(StartDirectory);
+	TrimFilenameFromPath(StartDirectory);
 
 	TSharedRef<SWindow> ModalWindow = SNew(SWindow)
 		.SupportsMinimize(false)
@@ -255,7 +255,7 @@ bool FSlateFileDlgWindow::OpenDirectoryDialog(const void* ParentWindowHandle, co
 	FString Filters = "";
 
 	FString StartDirectory = DefaultPath;
-	TrimStartDirectory(StartDirectory);
+	TrimFilenameFromPath(StartDirectory);
 
 	TSharedRef<SWindow> ModalWindow = SNew(SWindow)
 		.SupportsMinimize(false)
@@ -304,7 +304,7 @@ bool FSlateFileDlgWindow::SaveFileDialog(const void* ParentWindowHandle, const F
 	int32 DummyIndex;
 
 	FString StartDirectory = DefaultPath;
-	TrimStartDirectory(StartDirectory);
+	TrimFilenameFromPath(StartDirectory);
 
 	TSharedRef<SWindow> ModalWindow = SNew(SWindow)
 		.SupportsMinimize(false)
@@ -338,11 +338,11 @@ bool FSlateFileDlgWindow::SaveFileDialog(const void* ParentWindowHandle, const F
 	return (DialogWidget->GetResponse() == EResult::Accept && OutFilenames.Num() > 0);
 }
 
-void FSlateFileDlgWindow::TrimStartDirectory(FString &InPath)
+void FSlateFileDlgWindow::TrimFilenameFromPath(FString &InPath)
 {
-	if (InPath.Len() == 0)
+	if (InPath.Len() == 0 || !FPaths::FileExists(InPath))
 	{
-		// no path given. nothing to do.
+		// No path given OR no file portion to trim
 		return;
 	}
 
diff --git a/Engine/Source/Developer/SlateFileDialogs/Private/SlateFileDlgWindow.h b/Engine/Source/Developer/SlateFileDialogs/Private/SlateFileDlgWindow.h
--- a/Engine/Source/Developer/SlateFileDialogs/Private/SlateFileDlgWindow.h
+++ b/Engine/Source/Developer/SlateFileDialogs/Private/SlateFileDlgWindow.h
@@ -91,7 +91,7 @@ private:
 
 	FSlateFileDialogsStyle* StyleSet;
 
-	void TrimStartDirectory(FString &Path);
+	void TrimFilenameFromPath(FString &Path);
 };
 
 
