FDesktopPlatformWindows::OpenFileDialog based on FDesktopPlatformWindows::FileDialogShared manages natively extension like "prt.*"
but FDesktopPlatformLinux::OpenFileDialog uses FSlateFileDlgWindow that doesn't manage this kind of extension.
=> improvement of FSlateFileDialogVisitor::PassesFilterTest to fix it

diff --git a/Engine/Source/Developer/SlateFileDialogs/Private/SlateFileDlgWindow.cpp b/Engine/Source/Developer/SlateFileDialogs/Private/SlateFileDlgWindow.cpp
--- a/Engine/Source/Developer/SlateFileDialogs/Private/SlateFileDlgWindow.cpp
+++ b/Engine/Source/Developer/SlateFileDialogs/Private/SlateFileDlgWindow.cpp
@@ -137,7 +137,19 @@ public:
 			return true; // no filters. everything passes.
 		}
 
-		FString Extension = FPaths::GetExtension(FString(Filename), true);
+		FString BaseFile = Filename;
+		FString Extension = FPaths::GetExtension(BaseFile, true);
+		if (!Extension.IsEmpty() && FCString::IsNumeric(*Extension))
+		{
+			BaseFile.LeftChopInline(Extension.Len());
+			const int32 DotPos = BaseFile.Find(TEXT("."), ESearchCase::CaseSensitive, ESearchDir::FromEnd);
+			if (DotPos != INDEX_NONE)
+			{
+				Extension = BaseFile.Right(DotPos);
+				Extension += TEXT(".*");
+			}
+		}
+
 		// See if it matches any of the extensions
 		for (const FString& FilterExt : FilterList)
 		{
