Prevent SaveFileDialog's filename text box from updating when selecting a directory.

diff --git a/Engine/Source/Developer/SlateFileDialogs/Private/SlateFileDlgWindow.cpp b/Engine/Source/Developer/SlateFileDialogs/Private/SlateFileDlgWindow.cpp
--- a/Engine/Source/Developer/SlateFileDialogs/Private/SlateFileDlgWindow.cpp
+++ b/Engine/Source/Developer/SlateFileDialogs/Private/SlateFileDlgWindow.cpp
@@ -1547,7 +1547,11 @@ void SSlateFileOpenDlg::OnItemSelected(TSharedPtr<FFileEntry> Item, ESelectInfo:
 			FileList = Item->Label;
 		}
 
-		SetDefaultFile(FileList);
+		// Update file name text as long as we aren't saving a file with a directory selected
+		if (!(bSaveFile && Item->bIsDirectory))
+		{
+			SetDefaultFile(FileList);
+		}
 	}
 }
 
