Properly allow the file dialog window to have a passed-in parent.

diff --git a/Engine/Source/Developer/SlateFileDialogs/Private/SlateFileDlgWindow.cpp b/Engine/Source/Developer/SlateFileDialogs/Private/SlateFileDlgWindow.cpp
--- a/Engine/Source/Developer/SlateFileDialogs/Private/SlateFileDlgWindow.cpp
+++ b/Engine/Source/Developer/SlateFileDialogs/Private/SlateFileDlgWindow.cpp
@@ -215,6 +215,58 @@ FSlateFileDlgWindow::FSlateFileDlgWindow(FSlateFileDialogsStyle *InStyleSet)
 	StyleSet = InStyleSet;
 }
 
+namespace
+{
+	// Recursively adds all children of the given root window to the given list
+	void AddChildWindowsRecursively(const TSharedRef<SWindow>& RootWindow, TArray<TSharedRef<SWindow>>& OutWindows)
+	{
+		for (const TSharedRef<SWindow>& Window : RootWindow->GetChildWindows())
+		{
+			OutWindows.Add(Window);
+			AddChildWindowsRecursively(Window, OutWindows);
+		}
+	}
+
+	TArray<TSharedRef<SWindow>> GetTopLevelWindowsAndAllChildren()
+	{
+		TArray<TSharedRef<SWindow>> Result;
+
+		for (const TSharedRef<SWindow>& Window : FSlateApplication::Get().GetTopLevelWindows())
+		{
+			Result.Add(Window);
+			AddChildWindowsRecursively(Window, Result);
+		}
+
+		return Result;
+	}
+
+	// Defaults to root window if parent is not found
+	TSharedPtr<const SWidget> GetParentWindowWidget(const void* ParentWindowHandle)
+	{
+		if (ParentWindowHandle)
+		{
+			const TSharedRef<SWindow>* ParentWindow = GetTopLevelWindowsAndAllChildren().FindByPredicate(
+				[ParentWindowHandle] (const TSharedRef<SWindow> Window)
+				{
+					const TSharedPtr<FGenericWindow> NativeWindow = Window->GetNativeWindow();
+					if (NativeWindow.IsValid())
+					{
+						return NativeWindow->GetOSWindowHandle() == ParentWindowHandle;
+					}
+					return false;
+				}
+			);
+
+			if (ParentWindow)
+			{
+				return *ParentWindow;
+			}
+		}
+
+		return FGlobalTabmanager::Get()->GetRootWindow();
+	}
+}
+
 bool FSlateFileDlgWindow::OpenFileDialog(const void* ParentWindowHandle, const FString& DialogTitle, const FString& DefaultPath,
 		const FString& DefaultFile, const FString& FileTypes, uint32 Flags, TArray<FString>& OutFilenames, int32& OutFilterIndex)
 {
@@ -230,7 +282,7 @@ bool FSlateFileDlgWindow::OpenFileDialog(const void* ParentWindowHandle, const F
 		.MinWidth(600.0f)
 		.ActivationPolicy(EWindowActivationPolicy::Always)
 		.ClientSize(FVector2D(800, 500));
-	
+
 	DialogWidget = SNew(SSlateFileOpenDlg)
 		.bMultiSelectEnabled(Flags == 1)
 		.ParentWindow(ModalWindow)
@@ -238,14 +290,13 @@ bool FSlateFileDlgWindow::OpenFileDialog(const void* ParentWindowHandle, const F
 		.Filters(FileTypes)
 		.WindowTitleText(DialogTitle)
 		.StyleSet(StyleSet);
-	
+
 	DialogWidget->SetOutNames(&OutFilenames);
 	DialogWidget->SetOutFilterIndex(&OutFilterIndex);
-	
+
 	ModalWindow->SetContent( DialogWidget.ToSharedRef() );
-		
-	TSharedPtr<SWindow> RootWindow = FGlobalTabmanager::Get()->GetRootWindow();
-	FSlateApplication::Get().AddModalWindow(ModalWindow, RootWindow);
+
+	FSlateApplication::Get().AddModalWindow(ModalWindow, GetParentWindowWidget(ParentWindowHandle));
 
 	return (DialogWidget->GetResponse() == EResult::Accept && OutFilenames.Num() > 0);
 }
@@ -292,8 +343,7 @@ bool FSlateFileDlgWindow::OpenDirectoryDialog(const void* ParentWindowHandle, co
 
 	ModalWindow->SetContent( DialogWidget.ToSharedRef() );
 
-	TSharedPtr<SWindow> RootWindow = FGlobalTabmanager::Get()->GetRootWindow();
-	FSlateApplication::Get().AddModalWindow(ModalWindow, RootWindow);
+	FSlateApplication::Get().AddModalWindow(ModalWindow, GetParentWindowWidget(ParentWindowHandle));
 
 	bool RC = (DialogWidget->GetResponse() == EResult::Accept && TempOut.Num() > 0);
 
@@ -343,9 +393,8 @@ bool FSlateFileDlgWindow::SaveFileDialog(const void* ParentWindowHandle, const F
 	DialogWidget->SetDefaultFile(DefaultFile);
 
 	ModalWindow->SetContent( DialogWidget.ToSharedRef() );
-		
-	TSharedPtr<SWindow> RootWindow = FGlobalTabmanager::Get()->GetRootWindow();
-	FSlateApplication::Get().AddModalWindow(ModalWindow, RootWindow);
+
+	FSlateApplication::Get().AddModalWindow(ModalWindow, GetParentWindowWidget(ParentWindowHandle));
 
 	return (DialogWidget->GetResponse() == EResult::Accept && OutFilenames.Num() > 0);
 }
