Switched from 23314543 to using a predicate/enumerator API to solve the same problem.

diff --git a/Engine/Source/Developer/SlateFileDialogs/Private/SlateFileDlgWindow.cpp b/Engine/Source/Developer/SlateFileDialogs/Private/SlateFileDlgWindow.cpp
--- a/Engine/Source/Developer/SlateFileDialogs/Private/SlateFileDlgWindow.cpp
+++ b/Engine/Source/Developer/SlateFileDialogs/Private/SlateFileDlgWindow.cpp
@@ -217,52 +217,43 @@ FSlateFileDlgWindow::FSlateFileDlgWindow(FSlateFileDialogsStyle *InStyleSet)
 
 namespace
 {
-	// Recursively adds all children of the given root window to the given list
-	void AddChildWindowsRecursively(const TSharedRef<SWindow>& RootWindow, TArray<TSharedRef<SWindow>>& OutWindows)
+	TSharedPtr<SWindow> FindMatchingWindowImpl(const TArray<TSharedRef<SWindow>>& RootWindows, const TFunctionRef<bool(const TSharedRef<SWindow>&)> Predicate)
 	{
-		for (const TSharedRef<SWindow>& Window : RootWindow->GetChildWindows())
+		for (const TSharedRef<SWindow>& Window : RootWindows)
 		{
-			OutWindows.Add(Window);
-			AddChildWindowsRecursively(Window, OutWindows);
+			if (Predicate(Window))
+			{
+				return Window;
+			}
+
+			if (TSharedPtr<SWindow> ChildWindow = FindMatchingWindowImpl(Window->GetChildWindows(), Predicate))
+			{
+				return ChildWindow;
+			}
 		}
+		return nullptr;
 	}
 
-	TArray<TSharedRef<SWindow>> GetTopLevelWindowsAndAllChildren()
+	TSharedPtr<SWindow> FindMatchingWindow(const TFunctionRef<bool(const TSharedRef<SWindow>&)> Predicate)
 	{
-		TArray<TSharedRef<SWindow>> Result;
-
-		for (const TSharedRef<SWindow>& Window : FSlateApplication::Get().GetTopLevelWindows())
-		{
-			Result.Add(Window);
-			AddChildWindowsRecursively(Window, Result);
-		}
-
-		return Result;
+		return FindMatchingWindowImpl(FSlateApplication::Get().GetTopLevelWindows(), Predicate);
 	}
 
-	// Defaults to root window if parent is not found
 	TSharedPtr<const SWidget> GetParentWindowWidget(const void* ParentWindowHandle)
 	{
 		if (ParentWindowHandle)
 		{
-			const TSharedRef<SWindow>* ParentWindow = GetTopLevelWindowsAndAllChildren().FindByPredicate(
-				[ParentWindowHandle] (const TSharedRef<SWindow> Window)
-				{
-					const TSharedPtr<FGenericWindow> NativeWindow = Window->GetNativeWindow();
-					if (NativeWindow.IsValid())
-					{
-						return NativeWindow->GetOSWindowHandle() == ParentWindowHandle;
-					}
-					return false;
-				}
-			);
+			TSharedPtr<SWindow> ParentWindow = FindMatchingWindow([ParentWindowHandle](const TSharedRef<SWindow>& Window)
+			{
+				const TSharedPtr<FGenericWindow> NativeWindow = Window->GetNativeWindow();
+				return NativeWindow && NativeWindow->GetOSWindowHandle() == ParentWindowHandle;
+			});
 
 			if (ParentWindow)
 			{
-				return *ParentWindow;
+				return ParentWindow;
 			}
 		}
-
 		return FGlobalTabmanager::Get()->GetRootWindow();
 	}
 }
