Add support for specifying which file type is initially selected when creating a slate file dialog window.

diff --git a/Engine/Source/Developer/SlateFileDialogs/Private/SlateFileDialogsModule.cpp b/Engine/Source/Developer/SlateFileDialogs/Private/SlateFileDialogsModule.cpp
--- a/Engine/Source/Developer/SlateFileDialogs/Private/SlateFileDialogsModule.cpp
+++ b/Engine/Source/Developer/SlateFileDialogs/Private/SlateFileDialogsModule.cpp
@@ -29,18 +29,11 @@ void FSlateFileDialogsModule::ShutdownModule()
 
 
 bool FSlateFileDialogsModule::OpenFileDialog(const void* ParentWindowHandle, const FString& DialogTitle, const FString& DefaultPath,
-		const FString& DefaultFile, const FString& FileTypes, uint32 Flags, TArray<FString>& OutFilenames, int32& OutFilterIndex)
+		const FString& DefaultFile, const FString& FileTypes, uint32 Flags, TArray<FString>& OutFilenames, int32& OutFilterIndex,
+		int32 DefaultFilterIndex)
 {
 	FSlateFileDlgWindow dialog(FileDialogStyle);
-	return dialog.OpenFileDialog(ParentWindowHandle, DialogTitle, DefaultPath, DefaultFile, FileTypes, Flags, OutFilenames, OutFilterIndex);
-}
-
-
-bool FSlateFileDialogsModule::OpenFileDialog(const void* ParentWindowHandle, const FString& DialogTitle, const FString& DefaultPath,
-		const FString& DefaultFile, const FString& FileTypes, uint32 Flags, TArray<FString>& OutFilenames)
-{
-	FSlateFileDlgWindow dialog(FileDialogStyle);
-	return dialog.OpenFileDialog(ParentWindowHandle, DialogTitle, DefaultPath, DefaultFile, FileTypes, Flags, OutFilenames);
+	return dialog.OpenFileDialog(ParentWindowHandle, DialogTitle, DefaultPath, DefaultFile, FileTypes, Flags, OutFilenames, OutFilterIndex, DefaultFilterIndex);
 }
 
 
@@ -53,10 +46,11 @@ bool FSlateFileDialogsModule::OpenDirectoryDialog(const void* ParentWindowHandle
 
 
 bool FSlateFileDialogsModule::SaveFileDialog(const void* ParentWindowHandle, const FString& DialogTitle, const FString& DefaultPath,
-	const FString& DefaultFile, const FString& FileTypes, uint32 Flags, TArray<FString>& OutFilenames)
+		const FString& DefaultFile, const FString& FileTypes, uint32 Flags, TArray<FString>& OutFilenames, int32& OutFilterIndex,
+		int32 DefaultFilterIndex)
 {
 	FSlateFileDlgWindow dialog(FileDialogStyle);
-	return dialog.SaveFileDialog(ParentWindowHandle, DialogTitle, DefaultPath, DefaultFile, FileTypes, Flags, OutFilenames);
+	return dialog.SaveFileDialog(ParentWindowHandle, DialogTitle, DefaultPath, DefaultFile, FileTypes, Flags, OutFilenames, OutFilterIndex, DefaultFilterIndex);
 }
 
 
diff --git a/Engine/Source/Developer/SlateFileDialogs/Private/SlateFileDlgWindow.cpp b/Engine/Source/Developer/SlateFileDialogs/Private/SlateFileDlgWindow.cpp
--- a/Engine/Source/Developer/SlateFileDialogs/Private/SlateFileDlgWindow.cpp
+++ b/Engine/Source/Developer/SlateFileDialogs/Private/SlateFileDlgWindow.cpp
@@ -259,7 +259,8 @@ namespace
 }
 
 bool FSlateFileDlgWindow::OpenFileDialog(const void* ParentWindowHandle, const FString& DialogTitle, const FString& DefaultPath,
-		const FString& DefaultFile, const FString& FileTypes, uint32 Flags, TArray<FString>& OutFilenames, int32& OutFilterIndex)
+		const FString& DefaultFile, const FString& FileTypes, uint32 Flags, TArray<FString>& OutFilenames, int32& OutFilterIndex,
+		int32 DefaultFilterIndex)
 {
 	FString StartDirectory = DefaultPath;
 	TrimFilenameFromPath(StartDirectory);
@@ -284,6 +285,7 @@ bool FSlateFileDlgWindow::OpenFileDialog(const void* ParentWindowHandle, const F
 
 	DialogWidget->SetOutNames(&OutFilenames);
 	DialogWidget->SetOutFilterIndex(&OutFilterIndex);
+	DialogWidget->SetDefaultFilterIndex(DefaultFilterIndex);
 
 	ModalWindow->SetContent( DialogWidget.ToSharedRef() );
 
@@ -352,10 +354,9 @@ bool FSlateFileDlgWindow::OpenDirectoryDialog(const void* ParentWindowHandle, co
 
 
 bool FSlateFileDlgWindow::SaveFileDialog(const void* ParentWindowHandle, const FString& DialogTitle, const FString& DefaultPath,
-	const FString& DefaultFile, const FString& FileTypes, uint32 Flags, TArray<FString>& OutFilenames)
+	const FString& DefaultFile, const FString& FileTypes, uint32 Flags, TArray<FString>& OutFilenames, int32& OutFilterIndex,
+	int32 DefaultFilterIndex)
 {
-	int32 DummyIndex;
-
 	FString StartDirectory = DefaultPath;
 	TrimFilenameFromPath(StartDirectory);
 
@@ -380,7 +381,8 @@ bool FSlateFileDlgWindow::SaveFileDialog(const void* ParentWindowHandle, const F
 		.StyleSet(StyleSet);
 
 	DialogWidget->SetOutNames(&OutFilenames);
-	DialogWidget->SetOutFilterIndex(&DummyIndex);
+	DialogWidget->SetOutFilterIndex(&OutFilterIndex);
+	DialogWidget->SetDefaultFilterIndex(DefaultFilterIndex);
 	DialogWidget->SetDefaultFile(DefaultFile);
 
 	ModalWindow->SetContent( DialogWidget.ToSharedRef() );
@@ -390,6 +392,13 @@ bool FSlateFileDlgWindow::SaveFileDialog(const void* ParentWindowHandle, const F
 	return (DialogWidget->GetResponse() == EResult::Accept && OutFilenames.Num() > 0);
 }
 
+bool FSlateFileDlgWindow::SaveFileDialog(const void* ParentWindowHandle, const FString& DialogTitle, const FString& DefaultPath,
+	const FString& DefaultFile, const FString& FileTypes, uint32 Flags, TArray<FString>& OutFilenames)
+{
+	int32 DummyIndex;
+	return SaveFileDialog(ParentWindowHandle, DialogTitle, DefaultPath, DefaultFile, FileTypes, Flags, OutFilenames, DummyIndex);
+}
+
 void FSlateFileDlgWindow::TrimFilenameFromPath(FString &InPath)
 {
 	if (InPath.Len() == 0 || !FPaths::FileExists(InPath))
diff --git a/Engine/Source/Developer/SlateFileDialogs/Private/SlateFileDlgWindow.h b/Engine/Source/Developer/SlateFileDialogs/Private/SlateFileDlgWindow.h
--- a/Engine/Source/Developer/SlateFileDialogs/Private/SlateFileDlgWindow.h
+++ b/Engine/Source/Developer/SlateFileDialogs/Private/SlateFileDlgWindow.h
@@ -74,7 +74,8 @@ public:
 	FSlateFileDlgWindow(FSlateFileDialogsStyle *InStyleSet);
 
 	bool OpenFileDialog(const void* ParentWindowHandle, const FString& DialogTitle, const FString& DefaultPath,
-		const FString& DefaultFile, const FString& FileTypes, uint32 Flags, TArray<FString>& OutFilenames, int32& OutFilterIndex);
+		const FString& DefaultFile, const FString& FileTypes, uint32 Flags, TArray<FString>& OutFilenames, int32& OutFilterIndex,
+		int32 DefaultFilterIndex = 0);
 
 	bool OpenFileDialog(const void* ParentWindowHandle, const FString& DialogTitle, const FString& DefaultPath,
 		const FString& DefaultFile, const FString& FileTypes, uint32 Flags, TArray<FString>& OutFilenames);
@@ -82,6 +83,10 @@ public:
 	bool OpenDirectoryDialog(const void* ParentWindowHandle, const FString& DialogTitle, const FString& DefaultPath,
 		FString& OutFoldername);
 
+	bool SaveFileDialog(const void* ParentWindowHandle, const FString& DialogTitle, const FString& DefaultPath,
+		const FString& DefaultFile, const FString& FileTypes, uint32 Flags, TArray<FString>& OutFilenames, int32& OutFilterIndex,
+		int32 DefaultFilterIndex = 0);
+
 	bool SaveFileDialog(const void* ParentWindowHandle, const FString& DialogTitle, const FString& DefaultPath,
 		const FString& DefaultFile, const FString& FileTypes, uint32 Flags, TArray<FString>& OutFilenames);
 
@@ -148,6 +153,7 @@ public:
 	FSlateFileDlgWindow::EResult GetResponse() { return UserResponse; }
 	void SetOutNames(TArray<FString>* Ptr) { OutNames = Ptr; }
 	void SetOutFilterIndex(int32* InOutFilterIndex) { OutFilterIndex = InOutFilterIndex; }
+	void SetDefaultFilterIndex(int32 DefaultFilterIndex) { FilterIndex = DefaultFilterIndex; }
 	void SetDefaultFile(FString DefaultFile);
 
 private:	
diff --git a/Engine/Source/Developer/SlateFileDialogs/Public/ISlateFileDialogModule.h b/Engine/Source/Developer/SlateFileDialogs/Public/ISlateFileDialogModule.h
--- a/Engine/Source/Developer/SlateFileDialogs/Public/ISlateFileDialogModule.h
+++ b/Engine/Source/Developer/SlateFileDialogs/Public/ISlateFileDialogModule.h
@@ -17,7 +17,8 @@ public:
 	virtual ~ISlateFileDialogsModule( ) { }
 	
 	virtual bool OpenFileDialog(const void* ParentWindowHandle, const FString& DialogTitle, const FString& DefaultPath,
-		const FString& DefaultFile, const FString& FileTypes, uint32 Flags, TArray<FString>& OutFilenames, int32& OutFilterIndex)
+		const FString& DefaultFile, const FString& FileTypes, uint32 Flags, TArray<FString>& OutFilenames, int32& OutFilterIndex,
+		int32 DefaultFilterIndex = 0)
 	{
 		return false;
 	}
@@ -25,7 +26,8 @@ public:
 	virtual bool OpenFileDialog(const void* ParentWindowHandle, const FString& DialogTitle, const FString& DefaultPath,
 		const FString& DefaultFile, const FString& FileTypes, uint32 Flags, TArray<FString>& OutFilenames)
 	{
-		return false;
+		int32 DummyFilterIndex;
+		return OpenFileDialog(ParentWindowHandle, DialogTitle, DefaultPath, DefaultFile, FileTypes, Flags, OutFilenames, DummyFilterIndex);
 	}
 	
 	virtual bool OpenDirectoryDialog(const void* ParentWindowHandle, const FString& DialogTitle, const FString& DefaultPath,
@@ -35,10 +37,18 @@ public:
 	}
 	
 	virtual bool SaveFileDialog(const void* ParentWindowHandle, const FString& DialogTitle, const FString& DefaultPath,
-		const FString& DefaultFile, const FString& FileTypes, uint32 Flags, TArray<FString>& OutFilenames)
+		const FString& DefaultFile, const FString& FileTypes, uint32 Flags, TArray<FString>& OutFilenames, int32& OutFilterIndex,
+		int32 DefaultFilterIndex = 0)
 	{
 		return false;
 	}
+	
+	virtual bool SaveFileDialog(const void* ParentWindowHandle, const FString& DialogTitle, const FString& DefaultPath,
+		const FString& DefaultFile, const FString& FileTypes, uint32 Flags, TArray<FString>& OutFilenames)
+	{
+		int32 DummyFilterIndex;
+		return SaveFileDialog(ParentWindowHandle, DialogTitle, DefaultPath, DefaultFile, FileTypes, Flags, OutFilenames, DummyFilterIndex);
+	}
 
 	virtual ISlateFileDialogsModule* Get() { return nullptr; }
 };
diff --git a/Engine/Source/Developer/SlateFileDialogs/Public/SlateFileDialogs.h b/Engine/Source/Developer/SlateFileDialogs/Public/SlateFileDialogs.h
--- a/Engine/Source/Developer/SlateFileDialogs/Public/SlateFileDialogs.h
+++ b/Engine/Source/Developer/SlateFileDialogs/Public/SlateFileDialogs.h
@@ -19,19 +19,20 @@ public:
 public:
 	//ISlateFileDialogModule interface
 
-	bool OpenFileDialog(const void* ParentWindowHandle, const FString& DialogTitle, const FString& DefaultPath,
-		const FString& DefaultFile, const FString& FileTypes, uint32 Flags, TArray<FString>& OutFilenames, int32& OutFilterIndex);
+	using ISlateFileDialogsModule::OpenFileDialog;
+	virtual bool OpenFileDialog(const void* ParentWindowHandle, const FString& DialogTitle, const FString& DefaultPath,
+		const FString& DefaultFile, const FString& FileTypes, uint32 Flags, TArray<FString>& OutFilenames, int32& OutFilterIndex,
+		int32 DefaultFilterIndex = 0) override;
 
-	bool OpenFileDialog(const void* ParentWindowHandle, const FString& DialogTitle, const FString& DefaultPath,
-		const FString& DefaultFile, const FString& FileTypes, uint32 Flags, TArray<FString>& OutFilenames);
+	virtual bool OpenDirectoryDialog(const void* ParentWindowHandle, const FString& DialogTitle, const FString& DefaultPath,
+		FString& OutFoldername) override;
 
-	bool OpenDirectoryDialog(const void* ParentWindowHandle, const FString& DialogTitle, const FString& DefaultPath,
-		FString& OutFoldername);
+	using ISlateFileDialogsModule::SaveFileDialog;
+	virtual bool SaveFileDialog(const void* ParentWindowHandle, const FString& DialogTitle, const FString& DefaultPath,
+		const FString& DefaultFile, const FString& FileTypes, uint32 Flags, TArray<FString>& OutFilenames, int32& OutFilterIndex,
+		int32 DefaultFilterIndex = 0) override;
 
-	bool SaveFileDialog(const void* ParentWindowHandle, const FString& DialogTitle, const FString& DefaultPath,
-		const FString& DefaultFile, const FString& FileTypes, uint32 Flags, TArray<FString>& OutFilenames);
-
-	ISlateFileDialogsModule* Get();
+	virtual ISlateFileDialogsModule* Get() override;
 
 	FSlateFileDialogsStyle *GetFileDialogsStyle() { return FileDialogStyle; }
 
