Fixed up more bool-taking calls to take EAllowShrinking instead.

diff --git a/Engine/Plugins/Slate/SlateInsights/Source/SlateInsights/Private/SlateProvider.cpp b/Engine/Plugins/Slate/SlateInsights/Source/SlateInsights/Private/SlateProvider.cpp
--- a/Engine/Plugins/Slate/SlateInsights/Source/SlateInsights/Private/SlateProvider.cpp
+++ b/Engine/Plugins/Slate/SlateInsights/Source/SlateInsights/Private/SlateProvider.cpp
@@ -206,7 +206,7 @@ void FSlateProvider::ProcessWidgetUpdateSteps(const UE::Trace::IAnalyzer::FEvent
 	{
 		while (WidgetUpdateStepsEventIndexes.Num())
 		{
-			const TTuple<uint64, Message::FWidgetId> LastEntry = WidgetUpdateStepsEventIndexes.Pop(false);
+			const TTuple<uint64, Message::FWidgetId> LastEntry = WidgetUpdateStepsEventIndexes.Pop(EAllowShrinking::No);
 			WidgetPaintTimelines.EndEvent(LastEntry.Get<0>(), std::numeric_limits<double>::max());
 		}
 		bAcceptWidgetUpdateStepsComand = false;
@@ -260,7 +260,7 @@ void FSlateProvider::ProcessWidgetUpdateSteps(const UE::Trace::IAnalyzer::FEvent
 			const double Seconds = EventTime.AsSeconds(Cycle);
 			if (bAcceptWidgetUpdateStepsComand)
 			{
-				const TTuple<uint64, Message::FWidgetId> LastEntry = WidgetUpdateStepsEventIndexes.Pop(false);
+				const TTuple<uint64, Message::FWidgetId> LastEntry = WidgetUpdateStepsEventIndexes.Pop(EAllowShrinking::No);
 				WidgetPaintTimelines.EndEvent(LastEntry.Get<0>(), Seconds);
 			}
 		}
diff --git a/Engine/Source/Developer/SlateFileDialogs/Private/SlateFileDlgWindow.cpp b/Engine/Source/Developer/SlateFileDialogs/Private/SlateFileDlgWindow.cpp
--- a/Engine/Source/Developer/SlateFileDialogs/Private/SlateFileDlgWindow.cpp
+++ b/Engine/Source/Developer/SlateFileDialogs/Private/SlateFileDlgWindow.cpp
@@ -1027,7 +1027,7 @@ void SSlateFileOpenDlg::OnPathClicked(const FString& NewPath)
 
 	if ((History.Num()-HistoryIndex-1) > 0)
 	{
-		History.RemoveAt(HistoryIndex+1, History.Num()-HistoryIndex-1, true);
+		History.RemoveAt(HistoryIndex+1, History.Num()-HistoryIndex-1, EAllowShrinking::Yes);
 	}
 
 	History.Add(CurrentPath);
@@ -1045,7 +1045,7 @@ void SSlateFileOpenDlg::OnPathMenuItemClicked( FString ClickedPath )
 
 	if ((History.Num()-HistoryIndex-1) > 0)
 	{
-		History.RemoveAt(HistoryIndex+1, History.Num()-HistoryIndex-1, true);
+		History.RemoveAt(HistoryIndex+1, History.Num()-HistoryIndex-1, EAllowShrinking::Yes);
 	}
 
 	History.Add(CurrentPath);
@@ -1154,7 +1154,7 @@ FReply SSlateFileOpenDlg::OnQuickLinkClick(FSlateFileDlgWindow::EResult ButtonID
 	
 	if ((History.Num()-HistoryIndex-1) > 0)
 	{
-		History.RemoveAt(HistoryIndex+1, History.Num()-HistoryIndex-1, true);
+		History.RemoveAt(HistoryIndex+1, History.Num()-HistoryIndex-1, EAllowShrinking::Yes);
 	}
 
 	History.Add(CurrentPath);
@@ -1414,7 +1414,7 @@ void SSlateFileOpenDlg::OnItemDoubleClicked(TSharedPtr<FFileEntry> Item)
 
 		if ((History.Num()-HistoryIndex-1) > 0)
 		{
-			History.RemoveAt(HistoryIndex+1, History.Num()-HistoryIndex-1, true);
+			History.RemoveAt(HistoryIndex+1, History.Num()-HistoryIndex-1, EAllowShrinking::Yes);
 		}
 
 		History.Add(CurrentPath);
diff --git a/Engine/Source/Programs/CrashReportClient/Private/CrashReportAnalyticsSessionSummary.cpp b/Engine/Source/Programs/CrashReportClient/Private/CrashReportAnalyticsSessionSummary.cpp
--- a/Engine/Source/Programs/CrashReportClient/Private/CrashReportAnalyticsSessionSummary.cpp
+++ b/Engine/Source/Programs/CrashReportClient/Private/CrashReportAnalyticsSessionSummary.cpp
@@ -502,7 +502,7 @@ void FCrashReportAnalyticsSessionSummary::AppendLog(const TCHAR* Event)
 		}
 		else
 		{
-			DiagnosticLog.RemoveAt(0, EventLen - FreeLen, /*bAllowShrinking*/false); // Free space, remove the chars from the oldest events (in front).
+			DiagnosticLog.RemoveAt(0, EventLen - FreeLen, EAllowShrinking::No); // Free space, remove the chars from the oldest events (in front).
 		}
 	}
 
diff --git a/Engine/Source/Runtime/Analytics/Analytics/Public/AnalyticsConversion.h b/Engine/Source/Runtime/Analytics/Analytics/Public/AnalyticsConversion.h
--- a/Engine/Source/Runtime/Analytics/Analytics/Public/AnalyticsConversion.h
+++ b/Engine/Source/Runtime/Analytics/Analytics/Public/AnalyticsConversion.h
@@ -93,7 +93,7 @@ FString AnalyticsConversionToString(const TMap<KeyType, ValueType, Allocator, Ke
 		Result += TEXT(",");
 	}
 	// Remove the trailing comma (LeftChop will ensure an empty container won't crash here).
-	Result.LeftChopInline(1, false);
+	Result.LeftChopInline(1, EAllowShrinking::No);
 	return Result;
 }
 
diff --git a/Engine/Source/Runtime/Core/Private/HAL/MallocLeakDetection.cpp b/Engine/Source/Runtime/Core/Private/HAL/MallocLeakDetection.cpp
--- a/Engine/Source/Runtime/Core/Private/HAL/MallocLeakDetection.cpp
+++ b/Engine/Source/Runtime/Core/Private/HAL/MallocLeakDetection.cpp
@@ -103,7 +103,7 @@ void FMallocLeakDetection::PopContext()
 {
 	TArray<FContextString>* TLContexts = (TArray<FContextString>*)FPlatformTLS::GetTlsValue(FMallocLeakDetectionStatics::Get().ContextsTLSID);
 	check(TLContexts);
-	TLContexts->Pop(false);
+	TLContexts->Pop(EAllowShrinking::No);
 }
 
 void FMallocLeakDetection::AddCallstack(FCallstackTrack& Callstack)
diff --git a/Engine/Source/Runtime/Core/Private/Unix/UnixPlatformMisc.cpp b/Engine/Source/Runtime/Core/Private/Unix/UnixPlatformMisc.cpp
--- a/Engine/Source/Runtime/Core/Private/Unix/UnixPlatformMisc.cpp
+++ b/Engine/Source/Runtime/Core/Private/Unix/UnixPlatformMisc.cpp
@@ -491,11 +491,11 @@ namespace
 						ValueString.TrimStartAndEndInline();
 						if (ValueString.Left(1) == TEXT("\""))
 						{
-							ValueString.RightChopInline(1, false);
+							ValueString.RightChopInline(1, EAllowShrinking::No);
 						}
 						if (ValueString.Right(1) == TEXT("\""))
 						{
-							ValueString.LeftChopInline(1, false);
+							ValueString.LeftChopInline(1, EAllowShrinking::No);
 						}
 						Contents.Add(KeyString, ValueString);
 					}
diff --git a/Engine/Source/Runtime/Core/Private/Unix/UnixPlatformProcess.cpp b/Engine/Source/Runtime/Core/Private/Unix/UnixPlatformProcess.cpp
--- a/Engine/Source/Runtime/Core/Private/Unix/UnixPlatformProcess.cpp
+++ b/Engine/Source/Runtime/Core/Private/Unix/UnixPlatformProcess.cpp
@@ -1702,7 +1702,7 @@ FGenericPlatformProcess::EWaitAndForkResult FUnixPlatformProcess::WaitAndFork()
 						UE_LOG(LogHAL, Log, TEXT("[Parent] WaitAndFork child %d missing. Removing from children list..."), ChildPidAndSignal.Pid);
 					}
 
-					AllChildren.RemoveAt(ChildIdx, 1, false);
+					AllChildren.RemoveAt(ChildIdx, 1, EAllowShrinking::No);
 				}
 			}
 		}
@@ -2076,7 +2076,7 @@ bool FUnixPlatformProcess::IsFirstInstance()
 			FString ExecPath(FPlatformProcess::ExecutableName());
 			ExecPath.ReplaceInline(TEXT("/"), TEXT("-"), ESearchCase::CaseSensitive);
 			// [RCL] 2015-09-20: can run out of filename limits (256 bytes) due to a long path, be conservative and assume 4-char UTF-8 name like e.g. Japanese
-			ExecPath.RightInline(80, false);
+			ExecPath.RightInline(80, EAllowShrinking::No);
 
 			LockFileName += ExecPath;
 
diff --git a/Engine/Source/Runtime/CoreUObject/Private/UObject/UObjectGlobals.cpp b/Engine/Source/Runtime/CoreUObject/Private/UObject/UObjectGlobals.cpp
--- a/Engine/Source/Runtime/CoreUObject/Private/UObject/UObjectGlobals.cpp
+++ b/Engine/Source/Runtime/CoreUObject/Private/UObject/UObjectGlobals.cpp
@@ -2497,13 +2497,13 @@ namespace NameReuse
 				int32 Index = Names.IndexOfByPredicate([&](const TTuple<UPTRINT, FName>& Pair) { return Pair.Get<1>() == Name; });
 				if (Index != INDEX_NONE)
 				{
-					Names.RemoveAt(Index, 1, false);
+					Names.RemoveAt(Index, 1, EAllowShrinking::No);
 				}
 				else
 				{
 					if (Names.Num() >= MaxNamesPerEntry)
 					{
-						Names.RemoveAt(0, 1, false);
+						Names.RemoveAt(0, 1, EAllowShrinking::No);
 					}
 				}
 
@@ -2565,7 +2565,7 @@ namespace NameReuse
 			if (Index != Entries.Num() - 1)
 			{
 				FEntry Removed = MoveTemp(*Entry);
-				Entries.RemoveAt(Index, 1, false);
+				Entries.RemoveAt(Index, 1, EAllowShrinking::No);
 				Entries.Add(MoveTemp(Removed));
 			}
 		}
diff --git a/Engine/Source/Runtime/CoreUObject/Private/UObject/UObjectHash.cpp b/Engine/Source/Runtime/CoreUObject/Private/UObject/UObjectHash.cpp
--- a/Engine/Source/Runtime/CoreUObject/Private/UObject/UObjectHash.cpp
+++ b/Engine/Source/Runtime/CoreUObject/Private/UObject/UObjectHash.cpp
@@ -1065,7 +1065,7 @@ FORCEINLINE void RemoveFromClassMap(FUObjectHashTables& ThreadHash, UObjectBase*
 #if UE_STORE_OBJECT_LIST_INTERNAL_INDEX
 		ObjectList.Last()->ObjectListInternalIndex = Object->ObjectListInternalIndex;
 		ObjectList[Object->ObjectListInternalIndex] = ObjectList.Last();
-		ObjectList.Pop(false);
+		ObjectList.Pop(EAllowShrinking::No);
 #else
 		int32 NumRemoved = ObjectList.Remove(Object);
 
diff --git a/Engine/Source/Runtime/Engine/Private/Animation/AnimDataModel.cpp b/Engine/Source/Runtime/Engine/Private/Animation/AnimDataModel.cpp
--- a/Engine/Source/Runtime/Engine/Private/Animation/AnimDataModel.cpp
+++ b/Engine/Source/Runtime/Engine/Private/Animation/AnimDataModel.cpp
@@ -658,7 +658,7 @@ void ExtractPose(const TArray<FBoneAnimationTrack>& BoneAnimationTracks, const T
 					if (PoseBoneIndex == VB.VBIndex)
 					{
 						// Remove this bone as we have written data for it (false so we dont resize allocation)
-						VBCompactPoseData.RemoveAtSwap(Idx, 1, false);
+						VBCompactPoseData.RemoveAtSwap(Idx, 1, EAllowShrinking::No);
 						break; //Modified TArray so must break here
 					}
 				}
diff --git a/Engine/Source/Runtime/NavigationSystem/Private/NavMesh/RecastNavMeshGenerator.cpp b/Engine/Source/Runtime/NavigationSystem/Private/NavMesh/RecastNavMeshGenerator.cpp
--- a/Engine/Source/Runtime/NavigationSystem/Private/NavMesh/RecastNavMeshGenerator.cpp
+++ b/Engine/Source/Runtime/NavigationSystem/Private/NavMesh/RecastNavMeshGenerator.cpp
@@ -6518,7 +6518,7 @@ TSharedRef<FRecastTileGenerator> FRecastNavMeshGenerator::CreateTileGeneratorFro
 
 	TSharedRef<FRecastTileGenerator> TileGenerator = CreateTileGenerator(PendingElement.Coord, PendingElement.DirtyAreas, PendingElement.CreationTime);
 
-	PendingDirtyTiles.RemoveAt(PendingItemIdx, 1, false);
+	PendingDirtyTiles.RemoveAt(PendingItemIdx, 1, EAllowShrinking::No);
 
 	return TileGenerator;
 }
