Fix Linux Create Directory Button to correctly set the focus to the directory name box when pressed.

#rnx

diff --git a/Engine/Source/Developer/SlateFileDialogs/Private/SlateFileDlgWindow.cpp b/Engine/Source/Developer/SlateFileDialogs/Private/SlateFileDlgWindow.cpp
--- a/Engine/Source/Developer/SlateFileDialogs/Private/SlateFileDlgWindow.cpp
+++ b/Engine/Source/Developer/SlateFileDialogs/Private/SlateFileDlgWindow.cpp
@@ -1722,17 +1722,27 @@ void SSlateFileOpenDlg::OnNewDirectoryCommitted(const FText & InText, ETextCommi
 	}
 }
 
+void SSlateFileOpenDlg::OnNewDirectoryClickPostTick(float)
+{
+	if (FSlateApplication::Get().GetUserFocusedWidget(0) == NewDirectoryEditBox)
+	{
+		NewDirectoryEditBox->EnterEditingMode();
+
+		// The post tick event handler is not required anymore.
+		FSlateApplication::Get().OnPostTick().Remove(PostTickHandle);
+		PostTickHandle.Reset();
+	}
+}
 
 FReply SSlateFileOpenDlg::OnNewDirectoryClick()
 {
 	NewDirectorySizeBox->SetVisibility(EVisibility::Visible);
 	NewDirCancelButton->SetVisibility(EVisibility::Visible);
 	NewDirectoryEditBox->SetText(FString(""));
-	
-	FSlateApplication::Get().SetKeyboardFocus(NewDirectoryEditBox);	
-	NewDirectoryEditBox->EnterEditingMode();
-	
+	NewDirectoryEditBox->SetVisibility(EVisibility::Visible);
 	DirErrorMsg->SetVisibility(EVisibility::Collapsed);
+
+	PostTickHandle = FSlateApplication::Get().OnPostTick().AddSP(this, &SSlateFileOpenDlg::OnNewDirectoryClickPostTick);
 	
 	return FReply::Handled().SetUserFocus(NewDirectoryEditBox.ToSharedRef(), EFocusCause::SetDirectly);
 }
@@ -1769,7 +1779,7 @@ FReply SSlateFileOpenDlg::OnNewDirectoryAcceptCancelClick(FSlateFileDlgWindow::E
 	NewDirectorySizeBox->SetVisibility(EVisibility::Hidden);
 	NewDirCancelButton->SetVisibility(EVisibility::Hidden);
 	DirErrorMsg->SetVisibility(EVisibility::Collapsed);
-
+	NewDirectoryEditBox->SetVisibility(EVisibility::Hidden);
 	NewDirectoryEditBox->SetText(FString(""));
 
 	return FReply::Handled();
diff --git a/Engine/Source/Developer/SlateFileDialogs/Private/SlateFileDlgWindow.h b/Engine/Source/Developer/SlateFileDialogs/Private/SlateFileDlgWindow.h
--- a/Engine/Source/Developer/SlateFileDialogs/Private/SlateFileDlgWindow.h
+++ b/Engine/Source/Developer/SlateFileDialogs/Private/SlateFileDlgWindow.h
@@ -184,6 +184,8 @@ private:
 	bool OnNewDirectoryTextChanged(const FText &InText, FText &ErrorMsg);
 	FReply OnNewDirectoryAcceptCancelClick(FSlateFileDlgWindow::EResult ButtonID);
 
+	void OnNewDirectoryClickPostTick(float);
+
 	FReply OnGoForwardClick();
 	FReply OnGoBackClick();
 
@@ -250,6 +252,8 @@ private:
 
 	TArray<FString> History;
 	int32 HistoryIndex;
+
+	FDelegateHandle PostTickHandle;
 };
 
 
